fix for build errors related to Crash Reporter changes

diff --git a/Engine/Source/Runtime/CrashReportCore/CrashReportCore.Build.cs b/Engine/Source/Runtime/CrashReportCore/CrashReportCore.Build.cs
--- a/Engine/Source/Runtime/CrashReportCore/CrashReportCore.Build.cs
+++ b/Engine/Source/Runtime/CrashReportCore/CrashReportCore.Build.cs
@@ -15,23 +15,29 @@ public class CrashReportCore : ModuleRules
 		PublicDependencyModuleNames.AddRange(
 			new string[] {
 				"Core",
-                "CrashDebugHelper",
                 "XmlParser",
                 "Analytics",
                 "AnalyticsET",
-                 "HTTP",
+				"HTTP",
                 "Json",
            }
         );
 
-        PrecompileForTargets = PrecompileTargetsType.None;
-
-        if (Target.Type == TargetType.Game || Target.Type == TargetType.Client)
+        if (Target.Type == TargetType.Game || Target.Type == TargetType.Client || Target.Type == TargetType.Program)
         {
-            IsRedistributableOverride = true;
-        }
+			PrecompileForTargets = PrecompileTargetsType.None;
+			PublicDependencyModuleNames.Add("CrashDebugHelper");
+			PublicDefinitions.Add("WITH_CRASHREPORTER=1");
+
+			// add source control if we are targeting a program
+			if (Target.Type == TargetType.Program)
+			{
+				PublicDependencyModuleNames.Add("SourceControl");
+			}
+       	}
         else
         {
+			PublicDefinitions.Add("WITH_CRASHREPORTER=0");
             PublicDependencyModuleNames.Add("SourceControl");
         }
     }
diff --git a/Engine/Source/Runtime/CrashReportCore/Private/GenericErrorReport.cpp b/Engine/Source/Runtime/CrashReportCore/Private/GenericErrorReport.cpp
--- a/Engine/Source/Runtime/CrashReportCore/Private/GenericErrorReport.cpp
+++ b/Engine/Source/Runtime/CrashReportCore/Private/GenericErrorReport.cpp
@@ -10,9 +10,12 @@
 #include "Internationalization/Text.h"
 #include "CrashReportCoreConfig.h"
 #include "Modules/ModuleManager.h"
+#include "CrashReportUtil.h"
+
+#if WITH_CRASHREPORTER
 #include "CrashDebugHelper.h"
 #include "CrashDebugHelperModule.h"
-#include "CrashReportUtil.h"
+#endif
 
 // ----------------------------------------------------------------
 // Helpers
@@ -172,6 +175,7 @@ bool FGenericErrorReport::SetUserComment(const FText& UserComment)
 
 void FGenericErrorReport::SetPrimaryCrashProperties( FPrimaryCrashProperties& out_PrimaryCrashProperties )
 {
+#if WITH_CRASHREPORTER
 	FCrashDebugHelperModule& CrashHelperModule = FModuleManager::LoadModuleChecked<FCrashDebugHelperModule>( FName( "CrashDebugHelper" ) );
 	ICrashDebugHelper* Helper = CrashHelperModule.Get();
 	if (Helper && bValidCallstack)
@@ -197,6 +201,7 @@ void FGenericErrorReport::SetPrimaryCrashProperties( FPrimaryCrashProperties& ou
 
 		out_PrimaryCrashProperties.Save();
 	}
+#endif
 }
 
 void FGenericErrorReport::SetCrashReportClientVersion(const FString& InVersion)
diff --git a/Engine/Source/Runtime/CrashReportCore/Private/Mac/MacErrorReport.cpp b/Engine/Source/Runtime/CrashReportCore/Private/Mac/MacErrorReport.cpp
--- a/Engine/Source/Runtime/CrashReportCore/Private/Mac/MacErrorReport.cpp
+++ b/Engine/Source/Runtime/CrashReportCore/Private/Mac/MacErrorReport.cpp
@@ -2,19 +2,24 @@
 
 #include "Mac/MacErrorReport.h"
 #include "../CrashReportUtil.h"
-#include "CrashDebugHelperModule.h"
-#include "CrashDebugHelper.h"
 #include "Misc/FileHelper.h"
 #include "HAL/PlatformFilemanager.h"
 #include "Modules/ModuleManager.h"
+#include "CrashReportCoreConfig.h"
+
+#if WITH_CRASHREPORTER
+#include "CrashDebugHelper.h"
 #include "CrashDebugHelperModule.h"
+#endif
 
 #define LOCTEXT_NAMESPACE "CrashReportClient"
 
 namespace
 {
+#if WITH_CRASHREPORTER
 	/** Pointer to dynamically loaded crash diagnosis module */
 	FCrashDebugHelperModule* CrashHelperModule;
+#endif
 }
 
 FMacErrorReport::FMacErrorReport(const FString& Directory)
@@ -24,16 +29,21 @@ FMacErrorReport::FMacErrorReport(const FString& Directory)
 
 void FMacErrorReport::Init()
 {
+#if WITH_CRASHREPORTER
 	CrashHelperModule = &FModuleManager::LoadModuleChecked<FCrashDebugHelperModule>(FName("CrashDebugHelper"));
+#endif
 }
 
 void FMacErrorReport::ShutDown()
 {
+#if WITH_CRASHREPORTER
 	CrashHelperModule->ShutdownModule();
+#endif
 }
 
 FString FMacErrorReport::FindCrashedAppPath() const
 {
+#if WITH_CRASHREPORTER
 	TArray<uint8> Data;
 	if(FFileHelper::LoadFileToArray(Data, *(ReportDirectory / TEXT("Report.wer"))))
 	{
@@ -63,11 +73,13 @@ FString FMacErrorReport::FindCrashedAppPath() const
 	{
 		UE_LOG(LogStreaming, Error,	TEXT("Failed to read file '%s' error."),*(ReportDirectory / TEXT("Report.wer")));
 	}
+#endif
 	return "";
 }
 
 void FMacErrorReport::FindMostRecentErrorReports(TArray<FString>& ErrorReportPaths, const FTimespan& MaxCrashReportAge)
 {
+#if WITH_CRASHREPORTER
 	IPlatformFile& PlatformFile = FPlatformFileManager::Get().GetPlatformFile();
 
 	FDateTime MinCreationTime = FDateTime::UtcNow() - MaxCrashReportAge;
@@ -97,10 +109,12 @@ void FMacErrorReport::FindMostRecentErrorReports(TArray<FString>& ErrorReportPat
 
 		return TimeStampL > TimeStampR;
 	});
+#endif
 }
 
 FText FMacErrorReport::DiagnoseReport() const
 {
+#if WITH_CRASHREPORTER
 	// Should check if there are local PDBs before doing anything
 	ICrashDebugHelper* CrashDebugHelper = CrashHelperModule ? CrashHelperModule->Get() : nullptr;
 	if (!CrashDebugHelper)
@@ -142,6 +156,8 @@ FText FMacErrorReport::DiagnoseReport() const
 			return FText::FromString("Failed to create diagnosis information.");
 		}
 	}
+#endif
+	return FText::FromString(TEXT("Failed to load CrashDebugHelper."));
 }
 
 #undef LOCTEXT_NAMESPACE
diff --git a/Engine/Source/Runtime/CrashReportCore/Private/PendingReports.cpp b/Engine/Source/Runtime/CrashReportCore/Private/PendingReports.cpp
--- a/Engine/Source/Runtime/CrashReportCore/Private/PendingReports.cpp
+++ b/Engine/Source/Runtime/CrashReportCore/Private/PendingReports.cpp
@@ -1,6 +1,7 @@
 // Copyright 1998-2019 Epic Games, Inc. All Rights Reserved.
 
 #include "PendingReports.h"
+#include "CrashReportCoreModule.h"
 #include "HAL/PlatformFilemanager.h"
 #include "Logging/LogMacros.h"
 #include "HAL/FileManager.h"
diff --git a/Engine/Source/Runtime/CrashReportCore/Private/Windows/WindowsErrorReport.cpp b/Engine/Source/Runtime/CrashReportCore/Private/Windows/WindowsErrorReport.cpp
--- a/Engine/Source/Runtime/CrashReportCore/Private/Windows/WindowsErrorReport.cpp
+++ b/Engine/Source/Runtime/CrashReportCore/Private/Windows/WindowsErrorReport.cpp
@@ -1,10 +1,8 @@
 // Copyright 1998-2019 Epic Games, Inc. All Rights Reserved.
 
 #include "Windows/WindowsErrorReport.h"
-#include "CrashDebugHelperModule.h"
 #include "CrashReportUtil.h"
 #include "CrashDescription.h"
-#include "CrashDebugHelper.h"
 #include "Misc/FileHelper.h"
 #include "Modules/ModuleManager.h"
 #include "Internationalization/Text.h"
@@ -16,12 +14,19 @@
 #include <ShlObj.h>
 #include "Windows/HideWindowsPlatformTypes.h"
 
+#if WITH_CRASHREPORTER
+#include "CrashDebugHelperModule.h"
+#include "CrashDebugHelper.h"
+#endif
+
 #define LOCTEXT_NAMESPACE "CrashReportClient"
 
 namespace
 {
+#if WITH_CRASHREPORTER
 	/** Pointer to dynamically loaded crash diagnosis module */
 	FCrashDebugHelperModule* CrashHelperModule;
+#endif
 }
 
 /** Helper class used to parse specified string value based on the marker. */
@@ -64,12 +69,16 @@ FWindowsErrorReport::FWindowsErrorReport(const FString& Directory)
 
 void FWindowsErrorReport::Init()
 {
+#if WITH_CRASHREPORTER
 	CrashHelperModule = &FModuleManager::LoadModuleChecked<FCrashDebugHelperModule>(FName("CrashDebugHelper"));
+#endif
 }
 
 void FWindowsErrorReport::ShutDown()
 {
+#if WITH_CRASHREPORTER
 	CrashHelperModule->ShutdownModule();
+#endif
 }
 
 FString FWindowsErrorReport::FindCrashedAppPath() const
@@ -84,6 +93,7 @@ FText FWindowsErrorReport::DiagnoseReport() const
 	// Mark the callstack as invalid.
 	bValidCallstack = false;
 
+#if WITH_CRASHREPORTER
 	// Should check if there are local PDBs before doing anything
 	auto CrashDebugHelper = CrashHelperModule ? CrashHelperModule->Get() : nullptr;
 	if (!CrashDebugHelper)
@@ -108,6 +118,7 @@ FText FWindowsErrorReport::DiagnoseReport() const
 
 	// No longer required, only for backward compatibility, mark the callstack as valid.
 	bValidCallstack = true;
+#endif
 	return FText();
 }
 
@@ -136,6 +147,7 @@ static bool TryGetDirectoryCreationTimeUtc(const FString& InDirectoryName, FDate
 
 void FWindowsErrorReport::FindMostRecentErrorReports(TArray<FString>& ErrorReportPaths, const FTimespan& MaxCrashReportAge)
 {
+#if WITH_CRASHREPORTER
 	auto& PlatformFile = FPlatformFileManager::Get().GetPlatformFile();
 
 	FDateTime MinCreationTime = FDateTime::UtcNow() - MaxCrashReportAge;
@@ -174,6 +186,7 @@ void FWindowsErrorReport::FindMostRecentErrorReports(TArray<FString>& ErrorRepor
 
 		return CreationTimeL > CreationTimeR;
 	});
+#endif
 }
 
 #undef LOCTEXT_NAMESPACE
diff --git a/Engine/Source/Runtime/CrashReportCore/Public/PlatformErrorReport.h b/Engine/Source/Runtime/CrashReportCore/Public/PlatformErrorReport.h
--- a/Engine/Source/Runtime/CrashReportCore/Public/PlatformErrorReport.h
+++ b/Engine/Source/Runtime/CrashReportCore/Public/PlatformErrorReport.h
@@ -28,4 +28,8 @@ typedef FMacErrorReport FPlatformErrorReport;
 
 typedef FIOSErrorReport FPlatformErrorReport;
 
+#else
+
+typedef FGenericErrorReport FPlatformErrorReport;
+
 #endif // PLATFORM_LINUX
