Fix CIS error.

#rnx

diff --git a/Engine/Source/Runtime/CrashReportCore/CrashReportCore.Build.cs b/Engine/Source/Runtime/CrashReportCore/CrashReportCore.Build.cs
--- a/Engine/Source/Runtime/CrashReportCore/CrashReportCore.Build.cs
+++ b/Engine/Source/Runtime/CrashReportCore/CrashReportCore.Build.cs
@@ -23,7 +23,10 @@ public class CrashReportCore : ModuleRules
            }
         );
 
-        PrecompileForTargets = PrecompileTargetsType.None;
+		if (Target.Type != TargetType.Editor)
+        {
+            PrecompileForTargets = PrecompileTargetsType.None;
+        }
 
         if (Target.Type == TargetType.Game || Target.Type == TargetType.Client || Target.Type == TargetType.Program)
         {
