Add a portable callstack hash to *-runtime.xml and analytics crash reports

diff --git a/Engine/Source/Runtime/Core/Private/GenericPlatform/GenericPlatformCrashContext.cpp b/Engine/Source/Runtime/Core/Private/GenericPlatform/GenericPlatformCrashContext.cpp
--- a/Engine/Source/Runtime/Core/Private/GenericPlatform/GenericPlatformCrashContext.cpp
+++ b/Engine/Source/Runtime/Core/Private/GenericPlatform/GenericPlatformCrashContext.cpp
@@ -11,6 +11,7 @@
 #include "Misc/Optional.h"
 #include "Internationalization/Internationalization.h"
 #include "Misc/Guid.h"
+#include "Misc/SecureHash.h"
 #include "Containers/Ticker.h"
 #include "Misc/ConfigCacheIni.h"
 #include "Misc/CoreDelegates.h"
@@ -354,6 +355,7 @@ void FGenericCrashContext::SerializeContentToBuffer() const
 
 	// Add new portable callstack element with crash stack
 	AddPortableCallStack();
+	AddPortableCallStackHash();
 
 	AddCrashProperty( TEXT( "SourceContext" ), TEXT( "" ) );
 	AddCrashProperty( TEXT( "UserDescription" ), TEXT( "" ) );
@@ -514,6 +516,44 @@ void FGenericCrashContext::AddPlatformSpecificProperties() const
 	// @see FWindowsPlatformCrashContext::AddPlatformSpecificProperties
 }
 
+void FGenericCrashContext::AddPortableCallStackHash() const
+{
+	if (CallStack.Num() == 0)
+	{
+		AddCrashProperty(TEXT("PCallStackHash"), TEXT(""));
+		return;
+	}
+
+	// This may allocate if its the first time calling into this function
+	const TCHAR* ExeName = FPlatformProcess::ExecutableName();
+
+	// We dont want this to be thrown into an FString as it will alloc memory
+	const TCHAR* UE4EditorName = TEXT("UE4Editor");
+
+	FSHA1 Sha;
+	FSHAHash Hash;
+
+	for (TArray<FCrashStackFrame>::TConstIterator It(CallStack); It; ++It)
+	{
+		// If we are our own module or our module contains UE4Editor we assume we own these. We cannot depend on offsets of system libs
+		// as they may have different versions
+		if (It->ModuleName == ExeName || It->ModuleName.Contains(UE4EditorName))
+		{
+			Sha.Update(reinterpret_cast<const uint8*>(&It->Offset), sizeof(It->Offset));
+		}
+	}
+
+	Sha.Final();
+	Sha.GetHash(Hash.Hash);
+
+	FString EscapedPortableHash;
+
+	// Allocations here on both the ToString and AppendEscapedXMLString it self adds to the out FString
+	AppendEscapedXMLString(EscapedPortableHash, *Hash.ToString());
+
+	AddCrashProperty(TEXT("PCallStackHash"), *EscapedPortableHash);
+}
+
 void FGenericCrashContext::AddPortableCallStack() const
 {	
 
diff --git a/Engine/Source/Runtime/Core/Public/GenericPlatform/GenericPlatformCrashContext.h b/Engine/Source/Runtime/Core/Public/GenericPlatform/GenericPlatformCrashContext.h
--- a/Engine/Source/Runtime/Core/Public/GenericPlatform/GenericPlatformCrashContext.h
+++ b/Engine/Source/Runtime/Core/Public/GenericPlatform/GenericPlatformCrashContext.h
@@ -271,6 +271,9 @@ private:
 	/** Add callstack information to the crash report xml */
 	void AddPortableCallStack() const;
 
+    /** Produces a hash based on the offsets of the portable callstack and adds it to the xml */
+    void AddPortableCallStackHash() const;
+
 	/** Writes header information to the buffer. */
 	void AddHeader() const;
 
diff --git a/Engine/Source/Runtime/CrashReportCore/Private/CrashDescription.cpp b/Engine/Source/Runtime/CrashReportCore/Private/CrashDescription.cpp
--- a/Engine/Source/Runtime/CrashReportCore/Private/CrashDescription.cpp
+++ b/Engine/Source/Runtime/CrashReportCore/Private/CrashDescription.cpp
@@ -213,6 +213,7 @@ FString FPrimaryCrashProperties::EncodeArrayStringAsXMLString( const TArray<FStr
  * @EventParam AppDefaultLocale - The ICU default locale string or "en" if ICU is not enabled.
  * @EventParam UserActivityHint - Application-specific user activity string, if set in the crashed process. The meaning is game/app-specific.
  * @EventParam GameSessionID - Application-specific session Id, if set in the crashed process.
+ * @EventParam PCallStackHash - The hash of the portable callstack
  * @EventParam DeploymentName - Deployment name, also known as EpicApp. (e.g. "DevPlaytest", "PublicTest", "Live", etc)
  */
 void SendPreUploadEnsureAnalytics(const TArray<FAnalyticsEventAttribute>& InCrashAttributes)
@@ -245,6 +246,7 @@ void SendPreUploadEnsureAnalytics(const TArray<FAnalyticsEventAttribute>& InCras
  * @EventParam AppDefaultLocale - The ICU default locale string or "en" if ICU is not enabled.
  * @EventParam UserActivityHint - Application-specific user activity string, if set in the crashed process. The meaning is game/app-specific.
  * @EventParam GameSessionID - Application-specific session Id, if set in the crashed process.
+ * @EventParam PCallStackHash - The hash of the portable callstack
  * @EventParam DeploymentName - Deployment name, also known as EpicApp. (e.g. "DevPlaytest", "PublicTest", "Live", etc)
  */
 void SendPreUploadCrashAnalytics(const TArray<FAnalyticsEventAttribute>& InCrashAttributes)
@@ -390,6 +392,7 @@ void FPrimaryCrashProperties::MakeCrashEventAttributes(TArray<FAnalyticsEventAtt
 	OutCrashAttributes.Add(FAnalyticsEventAttribute(TEXT("UserActivityHint"), UserActivityHint.AsString()));
 	OutCrashAttributes.Add(FAnalyticsEventAttribute(TEXT("GameSessionID"), GameSessionID.AsString()));
 	OutCrashAttributes.Add(FAnalyticsEventAttribute(TEXT("DeploymentName"), DeploymentName));
+	OutCrashAttributes.Add(FAnalyticsEventAttribute(TEXT("PCallStackHash"), PCallStackHash));
 
 	// Add arbitrary engine data
 	const FXmlNode* EngineNode = XmlFile->GetRootNode()->FindChildNode( FGenericCrashContext::EngineDataTag );
@@ -469,6 +472,7 @@ FCrashContext::FCrashContext( const FString& CrashContextFilepath )
 		GetCrashProperty( bIsEnsure, FGenericCrashContext::RuntimePropertiesTag, TEXT("IsEnsure"));
 		GetCrashProperty( CrashType, FGenericCrashContext::RuntimePropertiesTag, TEXT("CrashType"));
 		GetCrashProperty( NumMinidumpFramesToIgnore, FGenericCrashContext::RuntimePropertiesTag, TEXT("NumMinidumpFramesToIgnore"));
+		GetCrashProperty( PCallStackHash, FGenericCrashContext::RuntimePropertiesTag, TEXT("PCallStackHash"));
 
 		if (CrashDumpMode == ECrashDumpMode::FullDump)
 		{
diff --git a/Engine/Source/Runtime/CrashReportCore/Public/CrashDescription.h b/Engine/Source/Runtime/CrashReportCore/Public/CrashDescription.h
--- a/Engine/Source/Runtime/CrashReportCore/Public/CrashDescription.h
+++ b/Engine/Source/Runtime/CrashReportCore/Public/CrashDescription.h
@@ -265,6 +265,13 @@ struct FPrimaryCrashProperties
 	 */
 	FCrashProperty GameSessionID;
 
+	/*
+	 * A hash representing a unique id for a portable callstack. These will be specific to the CL version of the application
+	 * @PCallStackHash	varchar(160)
+	 *
+	 */
+	FString PCallStackHash;
+
 	/**
 	 * Specifies the number of stack frames in the callstack to ignore when symbolicating from a minidump.
 	 */
