Fix for crash reporter upload loading a module not in the game thread.
Fix for corrupt crash context XML handling

diff --git a/Engine/Source/Runtime/CrashReportCore/Private/CrashDescription.cpp b/Engine/Source/Runtime/CrashReportCore/Private/CrashDescription.cpp
--- a/Engine/Source/Runtime/CrashReportCore/Private/CrashDescription.cpp
+++ b/Engine/Source/Runtime/CrashReportCore/Private/CrashDescription.cpp
@@ -170,8 +170,11 @@ void FPrimaryCrashProperties::ReadXML( const FString& CrashContextFilepath  )
 {
 	XmlFilepath = CrashContextFilepath;
 	XmlFile = new FXmlFile( XmlFilepath );
-	TimeOfCrash = FDateTime::UtcNow().GetTicks();
-	UpdateIDs();
+	if (XmlFile->IsValid())
+	{
+		TimeOfCrash = FDateTime::UtcNow().GetTicks();
+		UpdateIDs();
+	}
 }
 
 void FPrimaryCrashProperties::SetCrashGUID( const FString& Filepath )
@@ -395,31 +398,42 @@ void FPrimaryCrashProperties::MakeCrashEventAttributes(TArray<FAnalyticsEventAtt
 	OutCrashAttributes.Add(FAnalyticsEventAttribute(TEXT("PCallStackHash"), PCallStackHash));
 
 	// Add arbitrary engine data
-	const FXmlNode* EngineNode = XmlFile->GetRootNode()->FindChildNode( FGenericCrashContext::EngineDataTag );
-	if (EngineNode)
+	if (XmlFile->IsValid())
 	{
-		for (const FXmlNode* ChildNode : EngineNode->GetChildrenNodes())
+		const FXmlNode* EngineNode = XmlFile->GetRootNode()->FindChildNode( FGenericCrashContext::EngineDataTag );
+		if (EngineNode)
 		{
-			FString KeyName = FString(TEXT("EngineData.")) + ChildNode->GetTag();
-			OutCrashAttributes.Add(FAnalyticsEventAttribute(*KeyName, ChildNode->GetContent()));
+			for (const FXmlNode* ChildNode : EngineNode->GetChildrenNodes())
+			{
+				FString KeyName = FString(TEXT("EngineData.")) + ChildNode->GetTag();
+				OutCrashAttributes.Add(FAnalyticsEventAttribute(*KeyName, ChildNode->GetContent()));
+			}
 		}
-	}
 
-	// Add arbitrary game data
-	const FXmlNode* GameNode = XmlFile->GetRootNode()->FindChildNode( FGenericCrashContext::GameDataTag );
-	if (GameNode)
-	{
-		for (const FXmlNode* ChildNode : GameNode->GetChildrenNodes())
+		// Add arbitrary game data
+		const FXmlNode* GameNode = XmlFile->GetRootNode()->FindChildNode( FGenericCrashContext::GameDataTag );
+		if (GameNode)
 		{
-			FString KeyName = FString(TEXT("GameData.")) + ChildNode->GetTag();
-			OutCrashAttributes.Add(FAnalyticsEventAttribute(*KeyName, ChildNode->GetContent()));
+			for (const FXmlNode* ChildNode : GameNode->GetChildrenNodes())
+			{
+				FString KeyName = FString(TEXT("GameData.")) + ChildNode->GetTag();
+				OutCrashAttributes.Add(FAnalyticsEventAttribute(*KeyName, ChildNode->GetContent()));
+			}
 		}
 	}
+	else
+	{
+		FString KeyName = FString(TEXT("EngineData.InvalidXML"));
+		OutCrashAttributes.Add(FAnalyticsEventAttribute(*KeyName, true));
+	}
 }
 
 void FPrimaryCrashProperties::Save()
 {
-	XmlFile->Save( XmlFilepath );
+	if (XmlFile->IsValid())
+	{
+		XmlFile->Save( XmlFilepath );
+	}
 }
 
 /*-----------------------------------------------------------------------------
diff --git a/Engine/Source/Runtime/CrashReportCore/Public/CrashDescription.h b/Engine/Source/Runtime/CrashReportCore/Public/CrashDescription.h
--- a/Engine/Source/Runtime/CrashReportCore/Public/CrashDescription.h
+++ b/Engine/Source/Runtime/CrashReportCore/Public/CrashDescription.h
@@ -448,13 +448,16 @@ protected:
 	template <typename Type>
 	void GetCrashProperty( Type& out_ReadValue, const FString& MainCategory, const FString& SecondCategory ) const
 	{
-		const FXmlNode* MainNode = XmlFile->GetRootNode()->FindChildNode( MainCategory );
-		if (MainNode)
+		if (XmlFile->IsValid())
 		{
-			const FXmlNode* CategoryNode = MainNode->FindChildNode( SecondCategory );
-			if (CategoryNode)
+			const FXmlNode* MainNode = XmlFile->GetRootNode()->FindChildNode( MainCategory );
+			if (MainNode)
 			{
-				LexFromString( out_ReadValue, *FGenericCrashContext::UnescapeXMLString( CategoryNode->GetContent() ) );
+				const FXmlNode* CategoryNode = MainNode->FindChildNode( SecondCategory );
+				if (CategoryNode)
+				{
+					LexFromString( out_ReadValue, *FGenericCrashContext::UnescapeXMLString( CategoryNode->GetContent() ) );
+				}
 			}
 		}
 	}
@@ -469,19 +472,22 @@ protected:
 	/** Sets a crash property to a new value. */
 	void SetCrashProperty( const FString& MainCategory, const FString& SecondCategory, const FString& NewValue )
 	{
-		FXmlNode* MainNode = XmlFile->GetRootNode()->FindChildNode( MainCategory );
-		if (MainNode)
+		if (XmlFile->IsValid())
 		{
-			FXmlNode* CategoryNode = MainNode->FindChildNode( SecondCategory );
-			FString EscapedValue;
-			FGenericCrashContext::AppendEscapedXMLString(EscapedValue, *NewValue);
-			if (CategoryNode)
+			FXmlNode* MainNode = XmlFile->GetRootNode()->FindChildNode( MainCategory );
+			if (MainNode)
 			{
-				CategoryNode->SetContent( EscapedValue );
-			}
-			else
-			{
-				MainNode->AppendChildNode( SecondCategory, EscapedValue );
+				FXmlNode* CategoryNode = MainNode->FindChildNode( SecondCategory );
+				FString EscapedValue;
+				FGenericCrashContext::AppendEscapedXMLString(EscapedValue, *NewValue);
+				if (CategoryNode)
+				{
+					CategoryNode->SetContent( EscapedValue );
+				}
+				else
+				{
+					MainNode->AppendChildNode( SecondCategory, EscapedValue );
+				}
 			}
 		}
 	}
