fix for sending OOM crash with callstack
potential fix for occasional incorrect log with OOM
addition of StartTime and UpdateTime to key value data

diff --git a/Engine/Source/Developer/CrashDebugHelper/Private/CrashDebugHelperModule.cpp b/Engine/Source/Developer/CrashDebugHelper/Private/CrashDebugHelperModule.cpp
--- a/Engine/Source/Developer/CrashDebugHelper/Private/CrashDebugHelperModule.cpp
+++ b/Engine/Source/Developer/CrashDebugHelper/Private/CrashDebugHelperModule.cpp
@@ -46,3 +46,18 @@ ICrashDebugHelper* FCrashDebugHelperModule::Get()
 {
 	return CrashDebugHelper;
 }
+
+ICrashDebugHelper* FCrashDebugHelperModule::GetNew()
+{
+	if (CrashDebugHelper != nullptr)
+	{
+		delete CrashDebugHelper;
+		CrashDebugHelper = nullptr;
+	}
+	CrashDebugHelper = new FCrashDebugHelper();
+	if (CrashDebugHelper != nullptr)
+	{
+		CrashDebugHelper->Init();
+	}
+	return CrashDebugHelper;
+}
diff --git a/Engine/Source/Developer/CrashDebugHelper/Public/CrashDebugHelperModule.h b/Engine/Source/Developer/CrashDebugHelper/Public/CrashDebugHelperModule.h
--- a/Engine/Source/Developer/CrashDebugHelper/Public/CrashDebugHelperModule.h
+++ b/Engine/Source/Developer/CrashDebugHelper/Public/CrashDebugHelperModule.h
@@ -15,6 +15,9 @@ public:
 
 	/** Gets the crash debug helper singleton or returns NULL */
 	CRASHDEBUGHELPER_API ICrashDebugHelper* Get();
+	
+	/** Gets a new crash debug helper singleton for use with multiple report sending */
+	CRASHDEBUGHELPER_API ICrashDebugHelper* GetNew();
 
 private:
 	class ICrashDebugHelper* CrashDebugHelper;
diff --git a/Engine/Source/Runtime/CrashReportCore/Private/IOS/IOSErrorReport.cpp b/Engine/Source/Runtime/CrashReportCore/Private/IOS/IOSErrorReport.cpp
--- a/Engine/Source/Runtime/CrashReportCore/Private/IOS/IOSErrorReport.cpp
+++ b/Engine/Source/Runtime/CrashReportCore/Private/IOS/IOSErrorReport.cpp
@@ -102,7 +102,7 @@ void FIOSErrorReport::FindMostRecentErrorReports(TArray<FString>& ErrorReportPat
 FText FIOSErrorReport::DiagnoseReport() const
 {
 	// Should check if there are local PDBs before doing anything
-	ICrashDebugHelper* CrashDebugHelper = CrashHelperModule ? CrashHelperModule->Get() : nullptr;
+	ICrashDebugHelper* CrashDebugHelper = CrashHelperModule ? CrashHelperModule->GetNew() : nullptr;
 	if (!CrashDebugHelper)
 	{
 		// Not localized: should never be seen
