properly send crash analytic when we send a report

diff --git a/Engine/Source/Runtime/CrashReportCore/Private/CrashDescription.cpp b/Engine/Source/Runtime/CrashReportCore/Private/CrashDescription.cpp
--- a/Engine/Source/Runtime/CrashReportCore/Private/CrashDescription.cpp
+++ b/Engine/Source/Runtime/CrashReportCore/Private/CrashDescription.cpp
@@ -121,6 +121,7 @@ FPrimaryCrashProperties::FPrimaryCrashProperties()
 	, CrashReporterMessage( FGenericCrashContext::RuntimePropertiesTag, TEXT( "CrashReporterMessage" ), this )
 	, PlatformCallbackResult(FGenericCrashContext::PlatformPropertiesTag, TEXT("PlatformCallbackResult"), this)
 	, CrashReportClientVersion(FGenericCrashContext::RuntimePropertiesTag, TEXT("CrashReportClientVersion"), this)
+	, CPUBrand(FGenericCrashContext::RuntimePropertiesTag, TEXT("CPUBrand"), this)
 	, XmlFile( nullptr )
 {
 	CrashVersion = ECrashDescVersions::VER_1_NewCrashFormat;
@@ -396,6 +397,7 @@ void FPrimaryCrashProperties::MakeCrashEventAttributes(TArray<FAnalyticsEventAtt
 	OutCrashAttributes.Add(FAnalyticsEventAttribute(TEXT("GameSessionID"), GameSessionID.AsString()));
 	OutCrashAttributes.Add(FAnalyticsEventAttribute(TEXT("DeploymentName"), DeploymentName));
 	OutCrashAttributes.Add(FAnalyticsEventAttribute(TEXT("PCallStackHash"), PCallStackHash));
+	OutCrashAttributes.Add(FAnalyticsEventAttribute(TEXT("CPUBrand"), CPUBrand.AsString()));
 
 	// Add arbitrary engine data
 	if (XmlFile->IsValid())
diff --git a/Engine/Source/Runtime/CrashReportCore/Public/CrashDescription.h b/Engine/Source/Runtime/CrashReportCore/Public/CrashDescription.h
--- a/Engine/Source/Runtime/CrashReportCore/Public/CrashDescription.h
+++ b/Engine/Source/Runtime/CrashReportCore/Public/CrashDescription.h
@@ -383,6 +383,11 @@ struct FPrimaryCrashProperties
 	 */
 	FString CrashType;
 
+	/**
+	 *	The cpu brand of the device, e.g. Intel, iPhone6, etc.
+	 */
+	FCrashProperty CPUBrand;
+
 protected:
 	/** Default constructor. */
 	FPrimaryCrashProperties();
