additionally data added to analytic when sending a crash analytic

diff --git a/Engine/Source/Runtime/CrashReportCore/Private/CrashDescription.cpp b/Engine/Source/Runtime/CrashReportCore/Private/CrashDescription.cpp
--- a/Engine/Source/Runtime/CrashReportCore/Private/CrashDescription.cpp
+++ b/Engine/Source/Runtime/CrashReportCore/Private/CrashDescription.cpp
@@ -122,6 +122,8 @@ FPrimaryCrashProperties::FPrimaryCrashProperties()
 	, PlatformCallbackResult(FGenericCrashContext::PlatformPropertiesTag, TEXT("PlatformCallbackResult"), this)
 	, CrashReportClientVersion(FGenericCrashContext::RuntimePropertiesTag, TEXT("CrashReportClientVersion"), this)
 	, CPUBrand(FGenericCrashContext::RuntimePropertiesTag, TEXT("CPUBrand"), this)
+	, bIsOOM(false)
+	, bLowMemoryWarning(false)
 	, XmlFile( nullptr )
 {
 	CrashVersion = ECrashDescVersions::VER_1_NewCrashFormat;
@@ -398,6 +400,8 @@ void FPrimaryCrashProperties::MakeCrashEventAttributes(TArray<FAnalyticsEventAtt
 	OutCrashAttributes.Add(FAnalyticsEventAttribute(TEXT("DeploymentName"), DeploymentName));
 	OutCrashAttributes.Add(FAnalyticsEventAttribute(TEXT("PCallStackHash"), PCallStackHash));
 	OutCrashAttributes.Add(FAnalyticsEventAttribute(TEXT("CPUBrand"), CPUBrand.AsString()));
+	OutCrashAttributes.Add(FAnalyticsEventAttribute(TEXT("bIsOOM"), bIsOOM ? TEXT("true") : TEXT("false")));
+	OutCrashAttributes.Add(FAnalyticsEventAttribute(TEXT("bLowMemoryWarning"), bLowMemoryWarning ? TEXT("true") : TEXT("false")));
 
 	// Add arbitrary engine data
 	if (XmlFile->IsValid())
@@ -489,7 +493,9 @@ FCrashContext::FCrashContext( const FString& CrashContextFilepath )
 		GetCrashProperty( CrashType, FGenericCrashContext::RuntimePropertiesTag, TEXT("CrashType"));
 		GetCrashProperty( NumMinidumpFramesToIgnore, FGenericCrashContext::RuntimePropertiesTag, TEXT("NumMinidumpFramesToIgnore"));
 		GetCrashProperty( PCallStackHash, FGenericCrashContext::RuntimePropertiesTag, TEXT("PCallStackHash"));
-
+		GetCrashProperty( bIsOOM, FGenericCrashContext::RuntimePropertiesTag, TEXT("MemoryStats.bIsOOM"));
+		GetCrashProperty( bLowMemoryWarning, FGenericCrashContext::GameDataTag, TEXT("bLowMemoryCalled"));
+		
 		if (CrashDumpMode == ECrashDumpMode::FullDump)
 		{
 			// Set the full dump crash location when we have a full dump.
diff --git a/Engine/Source/Runtime/CrashReportCore/Public/CrashDescription.h b/Engine/Source/Runtime/CrashReportCore/Public/CrashDescription.h
--- a/Engine/Source/Runtime/CrashReportCore/Public/CrashDescription.h
+++ b/Engine/Source/Runtime/CrashReportCore/Public/CrashDescription.h
@@ -388,6 +388,16 @@ struct FPrimaryCrashProperties
 	 */
 	FCrashProperty CPUBrand;
 
+	/**
+	 * Whether it was an OOM or not
+	 */
+	bool bIsOOM;
+	
+	/**
+	 * Whether we got a low memory warning or not
+	 */
+	bool bLowMemoryWarning;
+	
 protected:
 	/** Default constructor. */
 	FPrimaryCrashProperties();
