fix for abnormal termination detection when the device is in the background

diff --git a/Engine/Source/Runtime/CrashReportCore/Private/CrashDescription.cpp b/Engine/Source/Runtime/CrashReportCore/Private/CrashDescription.cpp
--- a/Engine/Source/Runtime/CrashReportCore/Private/CrashDescription.cpp
+++ b/Engine/Source/Runtime/CrashReportCore/Private/CrashDescription.cpp
@@ -124,6 +124,7 @@ FPrimaryCrashProperties::FPrimaryCrashProperties()
 	, CPUBrand(FGenericCrashContext::RuntimePropertiesTag, TEXT("CPUBrand"), this)
 	, bIsOOM(false)
 	, bLowMemoryWarning(false)
+	, bInBackground(false)
 	, XmlFile( nullptr )
 {
 	CrashVersion = ECrashDescVersions::VER_1_NewCrashFormat;
@@ -409,6 +410,7 @@ void FPrimaryCrashProperties::MakeCrashEventAttributes(TArray<FAnalyticsEventAtt
 	OutCrashAttributes.Add(FAnalyticsEventAttribute(TEXT("CPUBrand"), CPUBrand.AsString()));
 	OutCrashAttributes.Add(FAnalyticsEventAttribute(TEXT("bIsOOM"), bIsOOM ? TEXT("true") : TEXT("false")));
 	OutCrashAttributes.Add(FAnalyticsEventAttribute(TEXT("bLowMemoryWarning"), bLowMemoryWarning ? TEXT("true") : TEXT("false")));
+	OutCrashAttributes.Add(FAnalyticsEventAttribute(TEXT("bInBackground"), bInBackground ? TEXT("true") : TEXT("false")));
 
 	// Add arbitrary engine data
 	if (XmlFile->IsValid())
@@ -512,7 +514,8 @@ void FCrashContext::SetupPrimaryCrashProperties()
 		GetCrashProperty( PCallStackHash, FGenericCrashContext::RuntimePropertiesTag, TEXT("PCallStackHash"));
 		GetCrashProperty( bIsOOM, FGenericCrashContext::RuntimePropertiesTag, TEXT("MemoryStats.bIsOOM"));
 		GetCrashProperty( bLowMemoryWarning, FGenericCrashContext::GameDataTag, TEXT("bLowMemoryCalled"));
-		
+		GetCrashProperty( bInBackground, FGenericCrashContext::GameDataTag, TEXT("bInBackground"));
+
 		if (CrashDumpMode == ECrashDumpMode::FullDump)
 		{
 			// Set the full dump crash location when we have a full dump.
diff --git a/Engine/Source/Runtime/CrashReportCore/Public/CrashDescription.h b/Engine/Source/Runtime/CrashReportCore/Public/CrashDescription.h
--- a/Engine/Source/Runtime/CrashReportCore/Public/CrashDescription.h
+++ b/Engine/Source/Runtime/CrashReportCore/Public/CrashDescription.h
@@ -398,6 +398,11 @@ struct FPrimaryCrashProperties
 	 */
 	bool bLowMemoryWarning;
 	
+	/**
+	 * Whether we were in the background when the crash happened
+	 */
+	bool bInBackground;
+	
 protected:
 	/** Default constructor. */
 	FPrimaryCrashProperties();
