Remove size limit for crash reporter

diff --git a/Engine/Source/Runtime/CrashReportCore/Private/CrashUpload.cpp b/Engine/Source/Runtime/CrashReportCore/Private/CrashUpload.cpp
--- a/Engine/Source/Runtime/CrashReportCore/Private/CrashUpload.cpp
+++ b/Engine/Source/Runtime/CrashReportCore/Private/CrashUpload.cpp
@@ -29,8 +29,6 @@
 namespace CrashUploadDefs
 {
 	const float PingTimeoutSeconds = 5.f;
-	// Ignore files bigger than 100MB; mini-dumps are smaller than this, but heap dumps can be very large
-	const int MaxFileSizeToUpload = 100 * 1024 * 1024;
 
 	const FString APIKey(TEXT("CrashReporter"));
 	const FString AppEnvironmentInternal(TEXT("Dev"));
@@ -178,12 +176,6 @@ bool FCrashUploadBase::CompressData(const TArray<FString>& InPendingFiles, FComp
 			continue;
 		}
 
-		if (IFileManager::Get().FileSize(*PathOfFileToUpload) > CrashUploadDefs::MaxFileSizeToUpload)
-		{
-			UE_LOG(CrashReportCoreLog, Warning, TEXT("Skipping large crash report file"));
-			continue;
-		}
-
 		if (!FFileHelper::LoadFileToArray(OutPostData, *PathOfFileToUpload))
 		{
 			UE_LOG(CrashReportCoreLog, Warning, TEXT("Failed to load crash report file"));
