Fix warning
#rnx

diff --git a/Engine/Source/Runtime/CrashReportCore/Private/Android/AndroidErrorReport.cpp b/Engine/Source/Runtime/CrashReportCore/Private/Android/AndroidErrorReport.cpp
--- a/Engine/Source/Runtime/CrashReportCore/Private/Android/AndroidErrorReport.cpp
+++ b/Engine/Source/Runtime/CrashReportCore/Private/Android/AndroidErrorReport.cpp
@@ -11,7 +11,7 @@
 
 #define LOCTEXT_NAMESPACE "CrashReport"
 
-namespace
+namespace AndroidErrorReport
 {
 	/** Pointer to dynamically loaded crash diagnosis module */
 	FCrashDebugHelperModule* CrashHelperModule;
@@ -19,12 +19,12 @@ namespace
 
 void FAndroidErrorReport::Init()
 {
-	CrashHelperModule = &FModuleManager::LoadModuleChecked<FCrashDebugHelperModule>(FName("CrashDebugHelper"));
+	AndroidErrorReport::CrashHelperModule = &FModuleManager::LoadModuleChecked<FCrashDebugHelperModule>(FName("CrashDebugHelper"));
 }
 
 void FAndroidErrorReport::ShutDown()
 {
-	CrashHelperModule->ShutdownModule();
+	AndroidErrorReport::CrashHelperModule->ShutdownModule();
 }
 
 #undef LOCTEXT_NAMESPACE
