Added bIsRequestingExit to FPrimaryCrashProperties

diff --git a/Engine/Source/Runtime/CrashReportCore/Private/CrashDescription.cpp b/Engine/Source/Runtime/CrashReportCore/Private/CrashDescription.cpp
--- a/Engine/Source/Runtime/CrashReportCore/Private/CrashDescription.cpp
+++ b/Engine/Source/Runtime/CrashReportCore/Private/CrashDescription.cpp
@@ -140,6 +140,7 @@ FPrimaryCrashProperties::FPrimaryCrashProperties()
 	, bIsOOM(false)
 	, bLowMemoryWarning(false)
 	, bInBackground(false)
+	, bIsRequestingExit(false)
 	, XmlFile( nullptr )
 {
 	CrashVersion = ECrashDescVersions::VER_1_NewCrashFormat;
@@ -426,6 +427,7 @@ void FPrimaryCrashProperties::MakeCrashEventAttributes(TArray<FAnalyticsEventAtt
 	OutCrashAttributes.Add(FAnalyticsEventAttribute(TEXT("bIsOOM"), bIsOOM ? TEXT("true") : TEXT("false")));
 	OutCrashAttributes.Add(FAnalyticsEventAttribute(TEXT("bLowMemoryWarning"), bLowMemoryWarning ? TEXT("true") : TEXT("false")));
 	OutCrashAttributes.Add(FAnalyticsEventAttribute(TEXT("bInBackground"), bInBackground ? TEXT("true") : TEXT("false")));
+	OutCrashAttributes.Add(FAnalyticsEventAttribute(TEXT("bIsRequestingExit"), bIsRequestingExit ? TEXT("true") : TEXT("false")));
 
 	// Add arbitrary engine data
 	if (XmlFile->IsValid())
@@ -530,6 +532,7 @@ void FCrashContext::SetupPrimaryCrashProperties()
 		GetCrashProperty( bIsOOM, FGenericCrashContext::RuntimePropertiesTag, TEXT("MemoryStats.bIsOOM"));
 		GetCrashProperty( bLowMemoryWarning, FGenericCrashContext::GameDataTag, TEXT("bLowMemoryCalled"));
 		GetCrashProperty( bInBackground, FGenericCrashContext::GameDataTag, TEXT("bInBackground"));
+		GetCrashProperty( bIsRequestingExit, FGenericCrashContext::GameDataTag, TEXT("IsRequestingExit"));
 
 		if (CrashDumpMode == ECrashDumpMode::FullDump)
 		{
diff --git a/Engine/Source/Runtime/CrashReportCore/Public/CrashDescription.h b/Engine/Source/Runtime/CrashReportCore/Public/CrashDescription.h
--- a/Engine/Source/Runtime/CrashReportCore/Public/CrashDescription.h
+++ b/Engine/Source/Runtime/CrashReportCore/Public/CrashDescription.h
@@ -430,7 +430,12 @@ struct FPrimaryCrashProperties
 	 * Whether we were in the background when the crash happened
 	 */
 	bool bInBackground;
-	
+
+	/**
+	 * Whether we crashed during shutdown
+	 */
+	bool bIsRequestingExit;
+
 protected:
 	/** Default constructor. */
 	FPrimaryCrashProperties();
