Remove PII only for none internal builds and avoid assuming user names are in the logs

diff --git a/Engine/Source/Runtime/CrashReportCore/Private/CrashDescription.cpp b/Engine/Source/Runtime/CrashReportCore/Private/CrashDescription.cpp
--- a/Engine/Source/Runtime/CrashReportCore/Private/CrashDescription.cpp
+++ b/Engine/Source/Runtime/CrashReportCore/Private/CrashDescription.cpp
@@ -170,9 +170,8 @@ void FPrimaryCrashProperties::UpdateIDs()
 		EpicAccountId = FString();
 	}
 
-	// Add real user name only if log files were allowed since the user name is in the log file and the user consented to sending this information.
-	const bool bSendUserName = FCrashReportCoreConfig::Get().GetSendLogFile() || FEngineBuildSettings::IsInternalBuild();
-	if (bSendUserName)
+	// Only send user name for internal builds
+	if (FEngineBuildSettings::IsInternalBuild())
 	{
 		// Remove periods from user names to match AutoReporter user names
 		// The name prefix is read by CrashRepository.AddNewCrash in the website code
diff --git a/Engine/Source/Runtime/CrashReportCore/Private/GenericErrorReport.cpp b/Engine/Source/Runtime/CrashReportCore/Private/GenericErrorReport.cpp
--- a/Engine/Source/Runtime/CrashReportCore/Private/GenericErrorReport.cpp
+++ b/Engine/Source/Runtime/CrashReportCore/Private/GenericErrorReport.cpp
@@ -5,6 +5,7 @@
 #include "HAL/PlatformFilemanager.h"
 #include "UObject/NameTypes.h"
 #include "Logging/LogMacros.h"
+#include "Misc/EngineBuildSettings.h"
 #include "Misc/Parse.h"
 #include "Misc/FileHelper.h"
 #include "Internationalization/Text.h"
@@ -12,6 +13,7 @@
 #include "Modules/ModuleManager.h"
 #include "CrashReportUtil.h"
 
+
 #if WITH_CRASHREPORTER
 #include "CrashDebugHelper.h"
 #include "CrashDebugHelperModule.h"
@@ -62,12 +64,11 @@ bool FGenericErrorReport::SetUserComment(const FText& UserComment)
 
 	FPrimaryCrashProperties::Get()->UserDescription = UserComment.ToString();
 
-	// Load the file and remove all PII if bAllowToBeContacted is set to false.
-	const bool bRemovePersonalData = !bAllowToBeContacted;
+	// Load the file and remove all PII if not an internal build
+	const bool bRemovePersonalData = !FEngineBuildSettings::IsInternalBuild();
 	if( bRemovePersonalData )
 	{
 		FPrimaryCrashProperties::Get()->UserName = TEXT( "" );
-		FPrimaryCrashProperties::Get()->EpicAccountId = TEXT( "" );
 		// For now remove the command line completely, to hide the potential personal data. Need to revisit it later.
 		FPrimaryCrashProperties::Get()->CommandLine = TEXT( "CommandLineRemoved" );
 	}
