Add CrashSignal and CrashSignalName to crash information collected
As well as send the CrashSignal as part of the analytic crash event

diff --git a/Engine/Source/Runtime/Core/Private/Unix/UnixPlatformCrashContext.cpp b/Engine/Source/Runtime/Core/Private/Unix/UnixPlatformCrashContext.cpp
--- a/Engine/Source/Runtime/Core/Private/Unix/UnixPlatformCrashContext.cpp
+++ b/Engine/Source/Runtime/Core/Private/Unix/UnixPlatformCrashContext.cpp
@@ -23,6 +23,7 @@
 #include "Misc/FeedbackContext.h"
 #include "Misc/App.h"
 #include "Misc/EngineVersion.h"
+#include "Misc/StringBuilder.h"
 #include "HAL/PlatformMallocCrash.h"
 #include "Unix/UnixPlatformRealTimeSignals.h"
 #include "Unix/UnixPlatformRunnableThread.h"
@@ -377,6 +378,23 @@ namespace UnixCrashReporterTracker
 	}
 }
 
+void FUnixCrashContext::AddPlatformSpecificProperties() const
+{
+	AddCrashProperty(TEXT("CrashSignal"), Signal);
+
+	ANSICHAR* AnsiSignalName = strsignal(Signal);
+	if (AnsiSignalName != nullptr)
+	{
+		TStringBuilder<32> SignalName;
+		SignalName.AppendAnsi(AnsiSignalName);
+
+		AddCrashProperty(TEXT("CrashSignalName"), *SignalName);
+	}
+	else
+	{
+		AddCrashProperty(TEXT("CrashSignalName"), TEXT("Unknown"));
+	}
+}
 
 void FUnixCrashContext::GenerateCrashInfoAndLaunchReporter(bool bReportingNonCrash) const
 {
@@ -715,6 +733,39 @@ extern int32 CORE_API GMaxNumberFileMappingCache;
 extern thread_local const TCHAR* GCrashErrorMessage;
 extern thread_local ECrashContextType GCrashErrorType;
 
+namespace
+{
+	ANSICHAR AnsiInternalBuffer[64] = { 0 };
+
+	// Taken from AndroidPlatformCrashContext.cpp, we need to avoid allocations while in the signal handler
+	const ANSICHAR* ItoANSI(uint64 Val, uint64 Base)
+	{
+		uint64 Index = 62;
+		int32 Pad = 0;
+		Base = FMath::Clamp<uint64>(Base, 2, 16);
+		if (Val)
+		{
+			for (; Val && Index; --Index, Val /= Base, --Pad)
+			{
+				AnsiInternalBuffer[Index] = "0123456789abcdef"[Val % Base];
+			}
+		}
+		else
+		{
+			AnsiInternalBuffer[Index--] = '0';
+			--Pad;
+		}
+
+		while (Pad > 0)
+		{
+			AnsiInternalBuffer[Index--] = '0';
+			--Pad;
+		}
+
+		return &AnsiInternalBuffer[Index + 1];
+	}
+}
+
 /** True system-specific crash handler that gets called first */
 void PlatformCrashHandler(int32 Signal, siginfo_t* Info, void* Context)
 {
@@ -730,12 +781,24 @@ void PlatformCrashHandler(int32 Signal, siginfo_t* Info, void* Context)
 	GMaxNumberFileMappingCache = 0;
 
 	ECrashContextType Type;
+	TStringBuilder<128> DefaultErrorMessage;
 	const TCHAR* ErrorMessage;
 
 	if (GCrashErrorMessage == nullptr)
 	{
 		Type = ECrashContextType::Crash;
-		ErrorMessage = TEXT("Caught signal");
+
+		DefaultErrorMessage.Append(TEXT("Caught signal "));
+		DefaultErrorMessage.AppendAnsi(ItoANSI(Signal, 10));
+
+		ANSICHAR* SignalName = strsignal(Signal);
+		if (SignalName != nullptr)
+		{
+			DefaultErrorMessage.Append(TEXT(" "));
+			DefaultErrorMessage.AppendAnsi(SignalName);
+		}
+
+		ErrorMessage = *DefaultErrorMessage;
 	}
 	else
 	{
diff --git a/Engine/Source/Runtime/Core/Public/Unix/UnixPlatformCrashContext.h b/Engine/Source/Runtime/Core/Public/Unix/UnixPlatformCrashContext.h
--- a/Engine/Source/Runtime/Core/Public/Unix/UnixPlatformCrashContext.h
+++ b/Engine/Source/Runtime/Core/Public/Unix/UnixPlatformCrashContext.h
@@ -107,6 +107,8 @@ struct CORE_API FUnixCrashContext : public FGenericCrashContext
 	 */
 	void SetFirstCrashHandlerFrame(uint64* ProgramCounter);
 
+	void AddPlatformSpecificProperties() const;
+
 protected:
 	/**
 	 * Dumps all the data from crash context to the "minidump" report.
diff --git a/Engine/Source/Runtime/CrashReportCore/Private/CrashDescription.cpp b/Engine/Source/Runtime/CrashReportCore/Private/CrashDescription.cpp
--- a/Engine/Source/Runtime/CrashReportCore/Private/CrashDescription.cpp
+++ b/Engine/Source/Runtime/CrashReportCore/Private/CrashDescription.cpp
@@ -149,6 +149,7 @@ FPrimaryCrashProperties::FPrimaryCrashProperties()
 	bHasLogFile = false;
 	bHasPrimaryData = false;
 	NumMinidumpFramesToIgnore = 0;
+	CrashSignal = 0;
 }
 
 void FPrimaryCrashProperties::Shutdown()
@@ -243,6 +244,7 @@ FString FPrimaryCrashProperties::EncodeArrayStringAsXMLString( const TArray<FStr
  * @EventParam UserActivityHint - Application-specific user activity string, if set in the crashed process. The meaning is game/app-specific.
  * @EventParam GameSessionID - Application-specific session Id, if set in the crashed process.
  * @EventParam PCallStackHash - The hash of the portable callstack
+ * @EventParam CrashSignal - The signal that was raised to enter the crash handler
  * @EventParam DeploymentName - Deployment name, also known as EpicApp. (e.g. "DevPlaytest", "PublicTest", "Live", etc)
  */
 void SendPreUploadEnsureAnalytics(const TArray<FAnalyticsEventAttribute>& InCrashAttributes)
@@ -276,6 +278,7 @@ void SendPreUploadEnsureAnalytics(const TArray<FAnalyticsEventAttribute>& InCras
  * @EventParam UserActivityHint - Application-specific user activity string, if set in the crashed process. The meaning is game/app-specific.
  * @EventParam GameSessionID - Application-specific session Id, if set in the crashed process.
  * @EventParam PCallStackHash - The hash of the portable callstack
+ * @EventParam CrashSignal - The signal that was raised to enter the crash handler
  * @EventParam DeploymentName - Deployment name, also known as EpicApp. (e.g. "DevPlaytest", "PublicTest", "Live", etc)
  */
 void SendPreUploadCrashAnalytics(const TArray<FAnalyticsEventAttribute>& InCrashAttributes)
@@ -422,6 +425,7 @@ void FPrimaryCrashProperties::MakeCrashEventAttributes(TArray<FAnalyticsEventAtt
 	OutCrashAttributes.Add(FAnalyticsEventAttribute(TEXT("GameSessionID"), GameSessionID.AsString()));
 	OutCrashAttributes.Add(FAnalyticsEventAttribute(TEXT("DeploymentName"), DeploymentName));
 	OutCrashAttributes.Add(FAnalyticsEventAttribute(TEXT("PCallStackHash"), PCallStackHash));
+	OutCrashAttributes.Add(FAnalyticsEventAttribute(TEXT("CrashSignal"), CrashSignal));
 	OutCrashAttributes.Add(FAnalyticsEventAttribute(TEXT("CPUBrand"), CPUBrand.AsString()));
 	OutCrashAttributes.Add(FAnalyticsEventAttribute(TEXT("bIsOOM"), bIsOOM ? TEXT("true") : TEXT("false")));
 	OutCrashAttributes.Add(FAnalyticsEventAttribute(TEXT("bLowMemoryWarning"), bLowMemoryWarning ? TEXT("true") : TEXT("false")));
@@ -528,6 +532,7 @@ void FCrashContext::SetupPrimaryCrashProperties()
 		GetCrashProperty( CrashType, FGenericCrashContext::RuntimePropertiesTag, TEXT("CrashType"));
 		GetCrashProperty( NumMinidumpFramesToIgnore, FGenericCrashContext::RuntimePropertiesTag, TEXT("NumMinidumpFramesToIgnore"));
 		GetCrashProperty( PCallStackHash, FGenericCrashContext::RuntimePropertiesTag, TEXT("PCallStackHash"));
+		GetCrashProperty( CrashSignal, FGenericCrashContext::PlatformPropertiesTag, TEXT("CrashSignal"));
 		GetCrashProperty( bIsOOM, FGenericCrashContext::RuntimePropertiesTag, TEXT("MemoryStats.bIsOOM"));
 		GetCrashProperty( bLowMemoryWarning, FGenericCrashContext::GameDataTag, TEXT("bLowMemoryCalled"));
 		GetCrashProperty( bInBackground, FGenericCrashContext::GameDataTag, TEXT("bInBackground"));
diff --git a/Engine/Source/Runtime/CrashReportCore/Public/CrashDescription.h b/Engine/Source/Runtime/CrashReportCore/Public/CrashDescription.h
--- a/Engine/Source/Runtime/CrashReportCore/Public/CrashDescription.h
+++ b/Engine/Source/Runtime/CrashReportCore/Public/CrashDescription.h
@@ -295,6 +295,11 @@ struct FPrimaryCrashProperties
 	 */
 	FString PCallStackHash;
 
+	/*
+	 * The signal that was raised to enter the crash handler
+	 */
+	int32 CrashSignal;
+
 	/**
 	 * Specifies the number of stack frames in the callstack to ignore when symbolicating from a minidump.
 	 */
