Fix for incorrect callstack in crash reports

When a ensure is reported before a crash the callstack from the crash was added to the ensure stack frames. This was because the debug helper instance used to analyze the minidump is a singleton and not designed to be reused. Use GetNew method to create a new instance of the helper, and delete the old one.

diff --git a/Engine/Source/Runtime/CrashReportCore/Private/Windows/WindowsErrorReport.cpp b/Engine/Source/Runtime/CrashReportCore/Private/Windows/WindowsErrorReport.cpp
--- a/Engine/Source/Runtime/CrashReportCore/Private/Windows/WindowsErrorReport.cpp
+++ b/Engine/Source/Runtime/CrashReportCore/Private/Windows/WindowsErrorReport.cpp
@@ -95,7 +95,7 @@ FText FWindowsErrorReport::DiagnoseReport() const
 
 #if WITH_CRASHREPORTER
 	// Should check if there are local PDBs before doing anything
-	auto CrashDebugHelper = CrashHelperModule ? CrashHelperModule->Get() : nullptr;
+	auto CrashDebugHelper = CrashHelperModule ? CrashHelperModule->GetNew() : nullptr;
 	if (!CrashDebugHelper)
 	{
 		// Not localized: should never be seen
