Wrapping CRC API URL in quotes to treat it as a string.

#rnx

diff --git a/Engine/Source/Programs/CrashReportClient/CrashReportClient.Target.cs b/Engine/Source/Programs/CrashReportClient/CrashReportClient.Target.cs
--- a/Engine/Source/Programs/CrashReportClient/CrashReportClient.Target.cs
+++ b/Engine/Source/Programs/CrashReportClient/CrashReportClient.Target.cs
@@ -53,7 +53,7 @@ public class CrashReportClientTarget : TargetRules
 		WindowsPlatform.bUseBundledDbgHelp = false;
 
 		// Add the definitions from config files
-		AddConfigMacro("CRC_TELEMETRY_URL=", TelemetryUrl);
+		AddConfigMacro("CRC_TELEMETRY_URL=", string.Format("\"{0}\"", TelemetryUrl));
 		AddConfigMacro("CRC_TELEMETRY_KEY_DEV=", TelemetryKey_Dev);
 		AddConfigMacro("CRC_TELEMETRY_KEY_RELEASE=", TelemetryKey_Release);
 
diff --git a/Engine/Source/Runtime/CrashReportCore/Private/CrashReportAnalytics.cpp b/Engine/Source/Runtime/CrashReportCore/Private/CrashReportAnalytics.cpp
--- a/Engine/Source/Runtime/CrashReportCore/Private/CrashReportAnalytics.cpp
+++ b/Engine/Source/Runtime/CrashReportCore/Private/CrashReportAnalytics.cpp
@@ -53,7 +53,7 @@ void FCrashReportAnalytics::Initialize()
 			!FEngineBuildSettings::IsInternalBuild();	// Internal Epic build
 
 		FAnalyticsET::Config Config;
-		Config.APIServerET = FString::Printf(TEXT("https://%s"), TEXT(PREPROCESSOR_TO_STRING(CRC_TELEMETRY_URL)));
+		Config.APIServerET = TEXT(CRC_TELEMETRY_URL);
 		Config.APIKeyET = bUseReleaseAccount ? TEXT(PREPROCESSOR_TO_STRING(CRC_TELEMETRY_KEY_RELEASE)) : TEXT(PREPROCESSOR_TO_STRING(CRC_TELEMETRY_KEY_DEV));
 
 	#else
