Remove UE4 strings and names from Crash reporting

diff --git a/Engine/Source/Programs/AutomationTool/Gauntlet/Framework/Gauntlet.Utils.cs b/Engine/Source/Programs/AutomationTool/Gauntlet/Framework/Gauntlet.Utils.cs
--- a/Engine/Source/Programs/AutomationTool/Gauntlet/Framework/Gauntlet.Utils.cs
+++ b/Engine/Source/Programs/AutomationTool/Gauntlet/Framework/Gauntlet.Utils.cs
@@ -1172,7 +1172,7 @@ namespace Gauntlet
 									Copied = true;
 
 									// Filter out some known unneeded files which can cause this warning, and log the message instead
-									string[] Blacklist = new string[]{ "UE4CC-XboxOne", "PersistentDownloadDir" };
+									string[] Blacklist = new string[]{ "UECC-XboxOne", "PersistentDownloadDir" };
 									string Message = string.Format("Long path file copy failed with {0}.  Please verify that this file is not required.", ex.Message);
 									if ( Blacklist.FirstOrDefault(B => { return SrcInfo.FullName.IndexOf(B, StringComparison.OrdinalIgnoreCase) >= 0; }) == null)
 									{
diff --git a/Engine/Source/Programs/AutomationTool/Gauntlet/Unreal/Base/Gauntlet.UnrealSession.cs b/Engine/Source/Programs/AutomationTool/Gauntlet/Unreal/Base/Gauntlet.UnrealSession.cs
--- a/Engine/Source/Programs/AutomationTool/Gauntlet/Unreal/Base/Gauntlet.UnrealSession.cs
+++ b/Engine/Source/Programs/AutomationTool/Gauntlet/Unreal/Base/Gauntlet.UnrealSession.cs
@@ -1088,8 +1088,8 @@ namespace Gauntlet
 			// testflights which append a long random name to the destination folder, easily pushing past 260 chars.
 			Func<string, string> TruncateLongPathFilter= (X =>
 			{
-				// /UE4CC-Windows-F0DD9BB04C3C9250FAF39D8AB4A88556/
-				Match M = Regex.Match(X, @"UE4CC-.+-([\dA-Fa-f]+)");
+				// /UECC-Windows-F0DD9BB04C3C9250FAF39D8AB4A88556/
+				Match M = Regex.Match(X, @"UECC-.+-([\dA-Fa-f]+)");
 
 				if (M.Success)
 				{
diff --git a/Engine/Source/Programs/CrashReportClient/Private/CrashReportClientApp.cpp b/Engine/Source/Programs/CrashReportClient/Private/CrashReportClientApp.cpp
--- a/Engine/Source/Programs/CrashReportClient/Private/CrashReportClientApp.cpp
+++ b/Engine/Source/Programs/CrashReportClient/Private/CrashReportClientApp.cpp
@@ -295,7 +295,7 @@ SubmitCrashReportResult RunWithUI(FPlatformErrorReport ErrorReport)
 	// Set up the main ticker
 	FMainLoopTiming MainLoop(IdealTickRate, EMainLoopOptions::UsingSlate);
 
-	// set the normal UE4 IsEngineExitRequested() when outer frame is closed
+	// set the normal IsEngineExitRequested() when outer frame is closed
 	FSlateApplication::Get().SetExitRequestedHandler(FSimpleDelegate::CreateStatic(&OnRequestExit));
 
 	// Prepare the custom Slate styles
@@ -827,7 +827,7 @@ static bool LoadTempCrashContextFromFile(FSharedCrashContext& CrashContext, uint
 	FindAndCopyValue(ContextProperties, TEXT("CommandLine"), SessionContext.CommandLine);
 	FindAndParseValue(ContextProperties, TEXT("LanguageLCID"), SessionContext.LanguageLCID);
 	FindAndCopyValue(ContextProperties, TEXT("AppDefaultLocale"), SessionContext.DefaultLocale);
-	FindAndParseValue(ContextProperties, TEXT("IsUE4Release"), SessionContext.bIsUE4Release);
+	FindAndParseValue(ContextProperties, TEXT("IsUERelease"), SessionContext.bIsUERelease);
 	FindAndCopyValue(ContextProperties, TEXT("UserName"), SessionContext.UserName);
 	FindAndCopyValue(ContextProperties, TEXT("BaseDir"), SessionContext.BaseDir);
 	FindAndCopyValue(ContextProperties, TEXT("RootDir"), SessionContext.RootDir);
diff --git a/Engine/Source/Programs/CrashReportClient/Private/Windows/CrashReportClientMainWindows.cpp b/Engine/Source/Programs/CrashReportClient/Private/Windows/CrashReportClientMainWindows.cpp
--- a/Engine/Source/Programs/CrashReportClient/Private/Windows/CrashReportClientMainWindows.cpp
+++ b/Engine/Source/Programs/CrashReportClient/Private/Windows/CrashReportClientMainWindows.cpp
@@ -129,7 +129,7 @@ int WINAPI WinMain(_In_ HINSTANCE hInInstance, _In_opt_ HINSTANCE hPrevInstance,
 
 			if (Handle.IsValid())
 			{
-				FString PidPathname = FString::Printf(TEXT("%sue4-crc-pid-%d"), FPlatformProcess::UserTempDir(), MonitoredEditorPid);
+				FString PidPathname = FString::Printf(TEXT("%sue-crc-pid-%d"), FPlatformProcess::UserTempDir(), MonitoredEditorPid);
 				if (TUniquePtr<FArchive> Ar = TUniquePtr<FArchive>(IFileManager::Get().CreateFileWriter(*PidPathname, FILEWRITE_EvenIfReadOnly)))
 				{
 					*Ar << RespawnPid;
diff --git a/Engine/Source/Runtime/Core/Private/GenericPlatform/GenericPlatformCrashContext.cpp b/Engine/Source/Runtime/Core/Private/GenericPlatform/GenericPlatformCrashContext.cpp
--- a/Engine/Source/Runtime/Core/Private/GenericPlatform/GenericPlatformCrashContext.cpp
+++ b/Engine/Source/Runtime/Core/Private/GenericPlatform/GenericPlatformCrashContext.cpp
@@ -48,7 +48,7 @@ const TCHAR* const FGenericCrashContext::CrashConfigFileNameW = TEXT("CrashRepor
 const TCHAR* const FGenericCrashContext::CrashConfigExtension = TEXT(".ini");
 const TCHAR* const FGenericCrashContext::ConfigSectionName = TEXT("CrashReportClient");
 const TCHAR* const FGenericCrashContext::CrashConfigPurgeDays = TEXT("CrashConfigPurgeDays");
-const TCHAR* const FGenericCrashContext::CrashGUIDRootPrefix = TEXT("UE4CC-");
+const TCHAR* const FGenericCrashContext::CrashGUIDRootPrefix = TEXT("UECC-");
 
 const TCHAR* const FGenericCrashContext::CrashContextExtension = TEXT(".runtime-xml");
 const TCHAR* const FGenericCrashContext::RuntimePropertiesTag = TEXT( "RuntimeProperties" );
@@ -56,7 +56,7 @@ const TCHAR* const FGenericCrashContext::PlatformPropertiesTag = TEXT( "Platform
 const TCHAR* const FGenericCrashContext::EngineDataTag = TEXT( "EngineData" );
 const TCHAR* const FGenericCrashContext::GameDataTag = TEXT( "GameData" );
 const TCHAR* const FGenericCrashContext::EnabledPluginsTag = TEXT("EnabledPlugins");
-const TCHAR* const FGenericCrashContext::UE4MinidumpName = TEXT( "UE4Minidump.dmp" );
+const TCHAR* const FGenericCrashContext::UEMinidumpName = TEXT( "UEMinidump.dmp" );
 const TCHAR* const FGenericCrashContext::NewLineTag = TEXT( "&nl;" );
 
 const TCHAR* const FGenericCrashContext::CrashTypeCrash = TEXT("Crash");
@@ -101,7 +101,7 @@ void FGenericCrashContext::Initialize()
 	NCached::Session.bIsSourceDistribution = FEngineBuildSettings::IsSourceDistribution();
 	NCached::Session.ProcessId = FPlatformProcess::GetCurrentProcessId();
 
-	NCached::Set(NCached::Session.GameName, *FString::Printf(TEXT("UE4-%s"), FApp::GetProjectName()));
+	NCached::Set(NCached::Session.GameName, *FString::Printf(TEXT("UE5-%s"), FApp::GetProjectName()));
 	NCached::Set(NCached::Session.GameSessionID, TEXT("")); // Updated by callback
 	NCached::Set(NCached::Session.GameStateName, TEXT("")); // Updated by callback
 	NCached::Set(NCached::Session.UserActivityHint, TEXT("")); // Updated by callback
@@ -132,7 +132,7 @@ void FGenericCrashContext::Initialize()
 	// Information that cannot be gathered if command line is not initialized (e.g. crash during static init)
 	if (FCommandLine::IsInitialized())
 	{
-		NCached::Session.bIsUE4Release = FApp::IsEngineInstalled();
+		NCached::Session.bIsUERelease = FApp::IsEngineInstalled();
 		NCached::Set(NCached::Session.CommandLine, (FCommandLine::IsInitialized() ? FCommandLine::GetOriginalForLogging() : TEXT("")));
 		NCached::Set(NCached::Session.EngineMode, FGenericPlatformMisc::GetEngineMode());
 		NCached::Set(NCached::Session.EngineModeEx, FGenericCrashContext::EngineModeExUnknown); // Updated from callback
@@ -220,7 +220,7 @@ PRAGMA_DISABLE_DEPRECATION_WARNINGS
 		}
 		if (InParams.bSetGameNameSuffix)
 		{
-			NCached::Set(NCached::Session.GameName, *(FString(TEXT("UE4-")) + FApp::GetProjectName() + InParams.GameNameSuffix));
+			NCached::Set(NCached::Session.GameName, *(FString(TEXT("UE-")) + FApp::GetProjectName() + InParams.GameNameSuffix));
 		}
 		if (InParams.SendUnattendedBugReports.IsSet())
 		{
@@ -496,7 +496,7 @@ void FGenericCrashContext::SerializeSessionContext(FString& Buffer)
 		const TCHAR* ProjectName = FApp::GetProjectName();
 		if (ProjectName != nullptr && ProjectName[0] != 0)
 		{
-			AddCrashPropertyInternal(Buffer, TEXT("GameName"), *FString::Printf(TEXT("UE4-%s"), ProjectName));
+			AddCrashPropertyInternal(Buffer, TEXT("GameName"), *FString::Printf(TEXT("UE-%s"), ProjectName));
 		}
 		else
 		{
@@ -535,7 +535,7 @@ void FGenericCrashContext::SerializeSessionContext(FString& Buffer)
 	AddCrashPropertyInternal(Buffer, TEXT("LanguageLCID"), NCached::Session.LanguageLCID);
 	AddCrashPropertyInternal(Buffer, TEXT("AppDefaultLocale"), NCached::Session.DefaultLocale);
 	AddCrashPropertyInternal(Buffer, TEXT("BuildVersion"), FApp::GetBuildVersion());
-	AddCrashPropertyInternal(Buffer, TEXT("IsUE4Release"), NCached::Session.bIsUE4Release);
+	AddCrashPropertyInternal(Buffer, TEXT("IsUERelease"), NCached::Session.bIsUERelease);
 
 	// Need to set this at the time of the crash to check if requesting exit had been called
 	AddCrashPropertyInternal(Buffer, TEXT("IsRequestingExit"), NCached::Session.bIsExitRequested);
@@ -777,7 +777,7 @@ void FGenericCrashContext::AddPortableCallStackHash() const
 	const TCHAR* ExeName = FPlatformProcess::ExecutableName();
 
 	// We dont want this to be thrown into an FString as it will alloc memory
-	const TCHAR* UE4EditorName = TEXT("UnrealEditor");
+	const TCHAR* UEEditorName = TEXT("UnrealEditor");
 
 	FSHA1 Sha;
 	FSHAHash Hash;
@@ -786,7 +786,7 @@ void FGenericCrashContext::AddPortableCallStackHash() const
 	{
 		// If we are our own module or our module contains UnrealEditor we assume we own these. We cannot depend on offsets of system libs
 		// as they may have different versions
-		if (It->ModuleName == ExeName || It->ModuleName.Contains(UE4EditorName))
+		if (It->ModuleName == ExeName || It->ModuleName.Contains(UEEditorName))
 		{
 			Sha.Update(reinterpret_cast<const uint8*>(&It->Offset), sizeof(It->Offset));
 		}
diff --git a/Engine/Source/Runtime/Core/Private/Misc/FileHelper.cpp b/Engine/Source/Runtime/Core/Private/Misc/FileHelper.cpp
--- a/Engine/Source/Runtime/Core/Private/Misc/FileHelper.cpp
+++ b/Engine/Source/Runtime/Core/Private/Misc/FileHelper.cpp
@@ -962,7 +962,8 @@ void FMaintenance::DeleteOldLogs()
 	// Remove all legacy crash contexts (regardless of age and purge settings, these are deprecated)
 	TArray<FString> Directories;
 	IFileManager::Get().FindFiles(Directories, *FString::Printf(TEXT("%s/UE4CC*"), *FPaths::ProjectLogDir()), false, true);
-
+	IFileManager::Get().FindFiles(Directories, *FString::Printf(TEXT("%s/UECC*"), *FPaths::ProjectLogDir()), false, true);
+	
 	for (const FString& Dir : Directories)
 	{
 		const FString CrashConfigDirectory = FPaths::ProjectLogDir() / Dir;
diff --git a/Engine/Source/Runtime/Core/Private/Windows/WindowsPlatformCrashContext.cpp b/Engine/Source/Runtime/Core/Private/Windows/WindowsPlatformCrashContext.cpp
--- a/Engine/Source/Runtime/Core/Private/Windows/WindowsPlatformCrashContext.cpp
+++ b/Engine/Source/Runtime/Core/Private/Windows/WindowsPlatformCrashContext.cpp
@@ -330,7 +330,7 @@ void FWindowsPlatformCrashContext::CopyPlatformSpecificFiles(const TCHAR* Output
 	LPEXCEPTION_POINTERS ExceptionInfo = (LPEXCEPTION_POINTERS)Context;
 	if (ExceptionInfo != nullptr)
 	{
-		const FString MinidumpFileName = FPaths::Combine(OutputDirectory, FGenericCrashContext::UE4MinidumpName);
+		const FString MinidumpFileName = FPaths::Combine(OutputDirectory, FGenericCrashContext::UEMinidumpName);
 		WriteMinidump(ProcessHandle.Get(), CrashedThreadId, *this, *MinidumpFileName, ExceptionInfo);
 	}
 
@@ -557,7 +557,7 @@ FProcHandle LaunchCrashReportClient(void** OutWritePipe, void** OutReadPipe, uin
 			
 			// The respanwned CRC process writes its own PID to a file named by this process PID (by parsing the -MONITOR={PID} arguments)
 			uint32 RepawnedCrcPid = 0;
-			FString PidFilePathname = FString::Printf(TEXT("%sue4-crc-pid-%d"), FPlatformProcess::UserTempDir(), FPlatformProcess::GetCurrentProcessId());
+			FString PidFilePathname = FString::Printf(TEXT("%sue-crc-pid-%d"), FPlatformProcess::UserTempDir(), FPlatformProcess::GetCurrentProcessId());
 			if (TUniquePtr<FArchive> Ar = TUniquePtr<FArchive>(IFileManager::Get().CreateFileReader(*PidFilePathname)))
 			{
 				*Ar << RepawnedCrcPid;
diff --git a/Engine/Source/Runtime/Core/Public/GenericPlatform/GenericPlatformCrashContext.h b/Engine/Source/Runtime/Core/Public/GenericPlatform/GenericPlatformCrashContext.h
--- a/Engine/Source/Runtime/Core/Public/GenericPlatform/GenericPlatformCrashContext.h
+++ b/Engine/Source/Runtime/Core/Public/GenericPlatform/GenericPlatformCrashContext.h
@@ -131,7 +131,7 @@ struct FSessionContext
 	bool 					bIsInternalBuild;
 	bool 					bIsPerforceBuild;
 	bool 					bIsSourceDistribution;
-	bool 					bIsUE4Release;
+	bool 					bIsUERelease;
 	bool					bIsOOM;
 	bool					bIsExitRequested;
 	uint32					ProcessId;
@@ -248,7 +248,7 @@ public:
 	static const TCHAR* const EngineDataTag;
 	static const TCHAR* const GameDataTag;
 	static const TCHAR* const EnabledPluginsTag;
-	static const TCHAR* const UE4MinidumpName;
+	static const TCHAR* const UEMinidumpName;
 	static const TCHAR* const NewLineTag;
 	static const int32 CrashGUIDLength = 128;
 
diff --git a/Engine/Source/Runtime/CrashReportCore/Private/CrashDescription.cpp b/Engine/Source/Runtime/CrashReportCore/Private/CrashDescription.cpp
--- a/Engine/Source/Runtime/CrashReportCore/Private/CrashDescription.cpp
+++ b/Engine/Source/Runtime/CrashReportCore/Private/CrashDescription.cpp
@@ -229,10 +229,10 @@ FString FPrimaryCrashProperties::EncodeArrayStringAsXMLString( const TArray<FStr
  *
  * @EventParam bHasPrimaryData - Whether the crash loaded data successfully from a crash context or legacy metadata file that was saved by the crashed process ("true" or "false")
  * @EventParam CrashVersion" - Describes the version of the crash data pipeline we used on the client side (1 = Legacy metadata based, 2 = unused, 3 = New crash context based)
- * @EventParam CrashGUID - GUID for this event generated by the crashed process (e.g. Windows format is UE4CC-Windows-1F07494140C2669B52AC96A1C4D20F83_0000 with the last four digits representing the index of event within a single session)
- * @EventParam PlatformCallbackResult - Platform-specific UE4 Core value. (integer)
- * @EventParam GameName - The name of the game that crashed from FApp::GetProjectName(). (e.g. UE4-QAGame)
- * @EventParam EngineVersion - The engine version of the crashed process from FEngineVersion::Current(). (e.g. 4.13.1-3142249+++UE4+Release-4.13)
+ * @EventParam CrashGUID - GUID for this event generated by the crashed process (e.g. Windows format is UECC-Windows-1F07494140C2669B52AC96A1C4D20F83_0000 with the last four digits representing the index of event within a single session)
+ * @EventParam PlatformCallbackResult - Platform-specific UE Core value. (integer)
+ * @EventParam GameName - The name of the game that crashed from FApp::GetProjectName(). (e.g. UE-QAGame)
+ * @EventParam EngineVersion - The engine version of the crashed process from FEngineVersion::Current(). (e.g. 4.13.1-3142249+++UE+Release-4.13)
  * @EventParam CrashReportClientVersion - CrashReportClient version, set in its config. (e.g. "1.0.0")
  * @EventParam LoginID - Unique string associated with the login account of the current machine from FPlatformMisc::GetLoginId().
  * @EventParam UserName - Username of the user of the process that crashed from FPlatformProcess::UserName().
@@ -263,10 +263,10 @@ void SendPreUploadEnsureAnalytics(const TArray<FAnalyticsEventAttribute>& InCras
  *
  * @EventParam bHasPrimaryData - Whether the crash loaded data successfully from a crash context or legacy metadata file that was saved by the crashed process ("true" or "false")
  * @EventParam CrashVersion" - Describes the version of the crash data pipeline we used on the client side (1 = Legacy metadata based, 2 = unused, 3 = New crash context based)
- * @EventParam CrashGUID - GUID for this event generated by the crashed process (e.g. Windows format is UE4CC-Windows-1F07494140C2669B52AC96A1C4D20F83_0000 with the last four digits representing the index of event within a single session)
- * @EventParam PlatformCallbackResult - Platform-specific UE4 Core value. (integer)
- * @EventParam GameName - The name of the game that crashed from FApp::GetProjectName(). (e.g. UE4-QAGame)
- * @EventParam EngineVersion - The engine version of the crashed process from FEngineVersion::Current(). (e.g. 4.13.1-3142249+++UE4+Release-4.13)
+ * @EventParam CrashGUID - GUID for this event generated by the crashed process (e.g. Windows format is UECC-Windows-1F07494140C2669B52AC96A1C4D20F83_0000 with the last four digits representing the index of event within a single session)
+ * @EventParam PlatformCallbackResult - Platform-specific UE Core value. (integer)
+ * @EventParam GameName - The name of the game that crashed from FApp::GetProjectName(). (e.g. UE-QAGame)
+ * @EventParam EngineVersion - The engine version of the crashed process from FEngineVersion::Current(). (e.g. 4.13.1-3142249+++UE+Release-4.13)
  * @EventParam CrashReportClientVersion - CrashReportClient version, set in its config. (e.g. "1.0.0")
  * @EventParam LoginID - Unique string associated with the login account of the current machine from FPlatformMisc::GetLoginId().
  * @EventParam UserName - Username of the user of the process that crashed from FPlatformProcess::UserName().
@@ -315,11 +315,11 @@ void FPrimaryCrashProperties::SendPreUploadAnalytics()
   *
   * @EventParam bHasPrimaryData - Whether the crash loaded data successfully from a crash context or legacy metadata file that was saved by the crashed process ("true" or "false")
   * @EventParam CrashVersion" - Describes the version of the crash data pipeline we used on the client side (1 = Legacy metadata based, 2 = unused, 3 = New crash context based)
-  * @EventParam CrashGUID - GUID for this event generated by the crashed process (e.g. Windows format is UE4CC-Windows-1F07494140C2669B52AC96A1C4D20F83_0000 with the last four digits representing the index of event within a single session)
-  * @EventParam PlatformCallbackResult - Platform-specific UE4 Core value. (integer)
+  * @EventParam CrashGUID - GUID for this event generated by the crashed process (e.g. Windows format is UECC-Windows-1F07494140C2669B52AC96A1C4D20F83_0000 with the last four digits representing the index of event within a single session)
+  * @EventParam PlatformCallbackResult - Platform-specific UE Core value. (integer)
   * @EventParam CrashType - The name of the crash report type (e.g Assert, Ensure, Hang)
-  * @EventParam GameName - The name of the game that crashed from FApp::GetProjectName(). (e.g. UE4-QAGame)
-  * @EventParam EngineVersion - The engine version of the crashed process from FEngineVersion::Current(). (e.g. 4.13.1-3142249+++UE4+Release-4.13)
+  * @EventParam GameName - The name of the game that crashed from FApp::GetProjectName(). (e.g. UE-QAGame)
+  * @EventParam EngineVersion - The engine version of the crashed process from FEngineVersion::Current(). (e.g. 4.13.1-3142249+++UE+Release-4.13)
   * @EventParam CrashReportClientVersion - CrashReportClient version, set in its config. (e.g. "1.0.0")
   * @EventParam LoginID - Unique string associated with the login account of the current machine from FPlatformMisc::GetLoginId().
   * @EventParam UserName - Username of the user of the process that crashed from FPlatformProcess::UserName().
@@ -351,11 +351,11 @@ void SendPostUploadEnsureAnalytics(const TArray<FAnalyticsEventAttribute>& InCra
  *
  * @EventParam bHasPrimaryData - Whether the crash loaded data successfully from a crash context or legacy metadata file that was saved by the crashed process ("true" or "false")
  * @EventParam CrashVersion" - Describes the version of the crash data pipeline we used on the client side (1 = Legacy metadata based, 2 = unused, 3 = New crash context based)
- * @EventParam CrashGUID - GUID for this event generated by the crashed process (e.g. Windows format is UE4CC-Windows-1F07494140C2669B52AC96A1C4D20F83_0000 with the last four digits representing the index of event within a single session)
- * @EventParam PlatformCallbackResult - Platform-specific UE4 Core value. (integer)
+ * @EventParam CrashGUID - GUID for this event generated by the crashed process (e.g. Windows format is UECC-Windows-1F07494140C2669B52AC96A1C4D20F83_0000 with the last four digits representing the index of event within a single session)
+ * @EventParam PlatformCallbackResult - Platform-specific UE Core value. (integer)
  * @EventParam CrashType - The name of the crash report type (e.g Assert, Ensure, Hang)
- * @EventParam GameName - The name of the game that crashed from FApp::GetProjectName(). (e.g. UE4-QAGame)
- * @EventParam EngineVersion - The engine version of the crashed process from FEngineVersion::Current(). (e.g. 4.13.1-3142249+++UE4+Release-4.13)
+ * @EventParam GameName - The name of the game that crashed from FApp::GetProjectName(). (e.g. UE-QAGame)
+ * @EventParam EngineVersion - The engine version of the crashed process from FEngineVersion::Current(). (e.g. 4.13.1-3142249+++UE+Release-4.13)
  * @EventParam CrashReportClientVersion - CrashReportClient version, set in its config. (e.g. "1.0.0")
  * @EventParam LoginID - Unique string associated with the login account of the current machine from FPlatformMisc::GetLoginId().
  * @EventParam UserName - Username of the user of the process that crashed from FPlatformProcess::UserName().
diff --git a/Engine/Source/Runtime/CrashReportCore/Private/CrashUpload.cpp b/Engine/Source/Runtime/CrashReportCore/Private/CrashUpload.cpp
--- a/Engine/Source/Runtime/CrashReportCore/Private/CrashUpload.cpp
+++ b/Engine/Source/Runtime/CrashReportCore/Private/CrashUpload.cpp
@@ -154,14 +154,14 @@ bool FCrashUploadBase::CompressData(const TArray<FString>& InPendingFiles, FComp
 	{
 		const FString Filename = FPaths::GetCleanFilename(PathOfFileToUpload);
 
-		const bool bValidFullDumpForCopy = Filename == FGenericCrashContext::UE4MinidumpName &&
+		const bool bValidFullDumpForCopy = Filename == FGenericCrashContext::UEMinidumpName &&
 			(FPrimaryCrashProperties::Get()->CrashDumpMode == ECrashDumpMode::FullDump || FPrimaryCrashProperties::Get()->CrashDumpMode == ECrashDumpMode::FullDumpAlways) &&
 			FPrimaryCrashProperties::Get()->CrashVersion >= ECrashDescVersions::VER_3_CrashContext &&
 			!FullCrashDumpLocation.IsEmpty();
 
 		if (bValidFullDumpForCopy)
 		{
-			const FString DestinationPath = FullCrashDumpLocation / FGenericCrashContext::UE4MinidumpName;
+			const FString DestinationPath = FullCrashDumpLocation / FGenericCrashContext::UEMinidumpName;
 			const bool bCreated = IFileManager::Get().MakeDirectory(*FullCrashDumpLocation, true);
 			if (!bCreated)
 			{
@@ -468,7 +468,7 @@ void FCrashUploadToReceiver::CompressAndSendData()
 
 	PendingFiles.Empty();
 
-	const FString Filename = ErrorReport.GetReportDirectoryLeafName() + TEXT(".ue4crash");
+	const FString Filename = ErrorReport.GetReportDirectoryLeafName() + TEXT(".uecrash");
 
 	// Set up request for upload
 	auto Request = CreateHttpRequest();
@@ -750,7 +750,7 @@ void FCrashUploadToDataRouter::CompressAndSendData()
 {
 	FCompressedHeader CompressedHeader;
 	CompressedHeader.DirectoryName = ErrorReport.GetReportDirectoryLeafName();
-	CompressedHeader.FileName = ErrorReport.GetReportDirectoryLeafName() + TEXT(".ue4crash");
+	CompressedHeader.FileName = ErrorReport.GetReportDirectoryLeafName() + TEXT(".uecrash");
 
 	FCompressedData CompressedData;
 	if (!CompressData(PendingFiles, CompressedData, PostData, &CompressedHeader))
diff --git a/Engine/Source/Runtime/CrashReportCore/Private/Windows/WindowsErrorReport.cpp b/Engine/Source/Runtime/CrashReportCore/Private/Windows/WindowsErrorReport.cpp
--- a/Engine/Source/Runtime/CrashReportCore/Private/Windows/WindowsErrorReport.cpp
+++ b/Engine/Source/Runtime/CrashReportCore/Private/Windows/WindowsErrorReport.cpp
@@ -156,7 +156,7 @@ void FWindowsErrorReport::FindMostRecentErrorReports(TArray<FString>& ErrorRepor
 		if (bIsDirectory)
 		{
 			FDateTime CreationTime;
-			if (TryGetDirectoryCreationTimeUtc(FilenameOrDirectory, CreationTime) && CreationTime > MinCreationTime && FCString::Strstr(FilenameOrDirectory, TEXT("UE4-")))
+			if (TryGetDirectoryCreationTimeUtc(FilenameOrDirectory, CreationTime) && CreationTime > MinCreationTime && FCString::Strstr(FilenameOrDirectory, TEXT("UE-")))
 			{
 				ErrorReportPaths.Add(FilenameOrDirectory);
 			}
diff --git a/Engine/Source/Runtime/CrashReportCore/Public/CrashDescription.h b/Engine/Source/Runtime/CrashReportCore/Public/CrashDescription.h
--- a/Engine/Source/Runtime/CrashReportCore/Public/CrashDescription.h
+++ b/Engine/Source/Runtime/CrashReportCore/Public/CrashDescription.h
@@ -37,7 +37,7 @@ struct FAnalyticsEventAttribute;
 	"CommandLine"
 	"LanguageLCID"
 	"AppDefaultLocale"
-	"IsUE4Release"
+	"IsUERelease"
 	"UserName"
 	"BaseDir"
 	"RootDir"
@@ -224,7 +224,7 @@ struct FPrimaryCrashProperties
 
 	/** 
 	 * Encoded engine version. (AppVersion)
-	 * E.g. 4.3.0.0-2215663+UE4-Releases+4.3
+	 * E.g. 4.3.0.0-2215663+UE-Releases+4.3
 	 * BuildVersion-BuiltFromCL-BranchName
 	 * @EngineVersion	varchar(64)	
 	 * 
@@ -372,7 +372,7 @@ struct FPrimaryCrashProperties
 	FCrashProperty CrashReporterMessage;
 
 	/**
-	 *	Platform-specific UE4 Core value (integer).
+	 *	Platform-specific UE Core value (integer).
 	 */
 	FCrashProperty PlatformCallbackResult;
 
diff --git a/Engine/Source/Runtime/CrashReportCore/Public/CrashReportCoreConfig.h b/Engine/Source/Runtime/CrashReportCore/Public/CrashReportCoreConfig.h
--- a/Engine/Source/Runtime/CrashReportCore/Public/CrashReportCoreConfig.h
+++ b/Engine/Source/Runtime/CrashReportCore/Public/CrashReportCoreConfig.h
@@ -18,8 +18,8 @@ class FConfigFile;
 /**
 *  Holds FullCrashDump properties from the config.
 *
-*	FullCrashDump_0_Branch=UE4
-*	FullCrashDump_0_Location=\\epicgames.net\root\Builds\UE4
+*	FullCrashDump_0_Branch=UE
+*	FullCrashDump_0_Location=\\epicgames.net\root\Builds\UE
 *	FullCrashDump_1_Branch=...
 *	...
 */
@@ -40,8 +40,8 @@ struct FFullCrashDumpEntry
 	const FString Location;
 
 	/**
-	*	Branch=UE4 means exact match
-	*	Branch=UE4* means contain match
+	*	Branch=UE means exact match
+	*	Branch=UE* means contain match
 	*/
 	const bool bExactMatch;
 };
diff --git a/Engine/Source/Runtime/CrashReportCore/Public/IOS/IOSErrorReport.h b/Engine/Source/Runtime/CrashReportCore/Public/IOS/IOSErrorReport.h
--- a/Engine/Source/Runtime/CrashReportCore/Public/IOS/IOSErrorReport.h
+++ b/Engine/Source/Runtime/CrashReportCore/Public/IOS/IOSErrorReport.h
@@ -35,7 +35,7 @@ public:
 
 	/**
 	 * Parse the callstack from the Apple-style crash report log
-	 * @return UE4 crash diagnosis text
+	 * @return UE crash diagnosis text
 	 */
 	FText DiagnoseReport() const;
 
diff --git a/Engine/Source/Runtime/CrashReportCore/Public/Mac/MacErrorReport.h b/Engine/Source/Runtime/CrashReportCore/Public/Mac/MacErrorReport.h
--- a/Engine/Source/Runtime/CrashReportCore/Public/Mac/MacErrorReport.h
+++ b/Engine/Source/Runtime/CrashReportCore/Public/Mac/MacErrorReport.h
@@ -35,7 +35,7 @@ public:
 
 	/**
 	 * Parse the callstack from the Apple-style crash report log
-	 * @return UE4 crash diagnosis text
+	 * @return UE crash diagnosis text
 	 */
 	FText DiagnoseReport() const;
 
