adding the option to hide the 'submit and restart' crash reporter option for platforms that do not support it.

#rnx

diff --git a/Engine/Source/Programs/CrashReportClient/Private/SCrashReportClient.cpp b/Engine/Source/Programs/CrashReportClient/Private/SCrashReportClient.cpp
--- a/Engine/Source/Programs/CrashReportClient/Private/SCrashReportClient.cpp
+++ b/Engine/Source/Programs/CrashReportClient/Private/SCrashReportClient.cpp
@@ -305,6 +305,7 @@ void SCrashReportClient::Construct(const FArguments& InArgs, const TSharedRef<FC
 					.Text(LOCTEXT("SendAndRestartEditor", "Send and Restart"))
 					.OnClicked(Client, &FCrashReportClient::SubmitAndRestart)
 					.IsEnabled(this, &SCrashReportClient::IsSendEnabled)
+					.Visibility( FCrashReportCoreConfig::Get().GetHideRestartOption() ? EVisibility::Collapsed : EVisibility::Visible )
 				]			
 			]
 		]
diff --git a/Engine/Source/Runtime/CrashReportCore/Private/CrashReportCoreConfig.cpp b/Engine/Source/Runtime/CrashReportCore/Private/CrashReportCoreConfig.cpp
--- a/Engine/Source/Runtime/CrashReportCore/Private/CrashReportCoreConfig.cpp
+++ b/Engine/Source/Runtime/CrashReportCore/Private/CrashReportCoreConfig.cpp
@@ -104,6 +104,7 @@ void FCrashReportCoreConfig::SetProjectConfigOverrides(const FConfigFile& InConf
 
 	// Default to false (show the option) when config is missing.
 	bHideLogFilesOption = false;
+	bHideRestartOption = false;
 
 	// Default to true (Allow the user to close without sending) when config is missing.
 	bIsAllowedToCloseWithoutSending = true;
@@ -117,6 +118,12 @@ void FCrashReportCoreConfig::SetProjectConfigOverrides(const FConfigFile& InConf
 			bHideLogFilesOption = FCString::ToBool(*HideLogFilesOptionValue->GetValue());
 		}
 
+		const FConfigValue* HideRestartOptionValue = Section->Find(TEXT("bHideRestartOption"));
+		if (HideRestartOptionValue != nullptr)
+		{
+			bHideRestartOption = FCString::ToBool(*HideRestartOptionValue->GetValue());
+		}
+
 		const FConfigValue* IsAllowedToCloseWithoutSendingValue = Section->Find(TEXT("bIsAllowedToCloseWithoutSending"));
 		if (IsAllowedToCloseWithoutSendingValue != nullptr)
 		{
diff --git a/Engine/Source/Runtime/CrashReportCore/Public/CrashReportCoreConfig.h b/Engine/Source/Runtime/CrashReportCore/Public/CrashReportCoreConfig.h
--- a/Engine/Source/Runtime/CrashReportCore/Public/CrashReportCoreConfig.h
+++ b/Engine/Source/Runtime/CrashReportCore/Public/CrashReportCoreConfig.h
@@ -94,6 +94,11 @@ struct FCrashReportCoreConfig
 		return bHideLogFilesOption;
 	}
 
+	const bool& GetHideRestartOption() const
+	{
+		return bHideRestartOption;
+	}
+
 	const bool& IsAllowedToCloseWithoutSending() const
 	{
 		return bIsAllowedToCloseWithoutSending;
@@ -153,6 +158,9 @@ protected:
 	/** Whether the user is shown the option to enable/disable sending the log file. */
 	bool bHideLogFilesOption;
 
+	/** Whether the user is shown the option to restart the crashed application */
+	bool bHideRestartOption;
+
 	/** Whether the user is allowed to close the crash reporter without sending a report */
 	bool bIsAllowedToCloseWithoutSending;
 
