Fixed CrashReportClient writing a buggy compressed header at the start of the compressed report and the valid one at the end of the report. (Contributed by [REDACTED_PERSON])
  - Added the 'CR1' bytes at the very beginning of the compressed data to let the backend know that CRC emits the proper header in front of the data stream.
  - Prevented CRC from writing a second header at the end of the compressed data.

#8246 : Fixed crash report client not properly writing the packet header

diff --git a/Engine/Source/Runtime/CrashReportCore/Private/CrashUpload.cpp b/Engine/Source/Runtime/CrashReportCore/Private/CrashUpload.cpp
--- a/Engine/Source/Runtime/CrashReportCore/Private/CrashUpload.cpp
+++ b/Engine/Source/Runtime/CrashReportCore/Private/CrashUpload.cpp
@@ -74,6 +74,11 @@ struct FCompressedHeader
 	/** Serialization operator. */
 	friend FArchive& operator << (FArchive& Ar, FCompressedHeader& Data)
 	{
+		// The 'CR1' marker prevents the data router backend to fallback to a backward compatibility version where
+		// a buggy/incomplete header was written at the beginning of the stream and the correct/valid one at the end.
+		uint8 Version[] = {'C', 'R', '1'};
+		Ar.Serialize(Version, sizeof(Version));
+
 		Data.DirectoryName.SerializeAsANSICharArray(Ar, 260);
 		Data.FileName.SerializeAsANSICharArray(Ar, 260);
 		Ar << Data.UncompressedSize;
@@ -227,7 +232,7 @@ bool FCrashUploadBase::CompressData(const TArray<FString>& InPendingFiles, FComp
 
 	if (OptionalHeader != nullptr)
 	{
-		FMemoryWriter MemoryHeaderWriter(UncompressedData, false, true);
+		FMemoryWriter MemoryHeaderWriter(UncompressedData);
 
 		OptionalHeader->UncompressedSize = UncompressedData.Num();
 		OptionalHeader->FileCount = CurrentFileIndex;
