Fixing issue with crashes that occur with -fullcrashdump enabled either uploading a corrupted minidump or skipping their upload entirely.

diff --git a/Engine/Source/Runtime/CrashReportCore/Private/CrashUpload.cpp b/Engine/Source/Runtime/CrashReportCore/Private/CrashUpload.cpp
--- a/Engine/Source/Runtime/CrashReportCore/Private/CrashUpload.cpp
+++ b/Engine/Source/Runtime/CrashReportCore/Private/CrashUpload.cpp
@@ -154,30 +154,36 @@ bool FCrashUploadBase::CompressData(const TArray<FString>& InPendingFiles, FComp
 
 	const FString FullCrashDumpLocation = FPrimaryCrashProperties::Get()->FullCrashDumpLocation.AsString();
 
+	bool bIsFullDumpCrash = (FPrimaryCrashProperties::Get()->CrashDumpMode == ECrashDumpMode::FullDump ||
+		FPrimaryCrashProperties::Get()->CrashDumpMode == ECrashDumpMode::FullDumpAlways) &&
+		FPrimaryCrashProperties::Get()->CrashVersion >= ECrashDescVersions::VER_3_CrashContext;
+
 	// Loop to keep trying files until a send succeeds or we run out of files
 	for (const FString& PathOfFileToUpload : InPendingFiles)
 	{
 		const FString Filename = FPaths::GetCleanFilename(PathOfFileToUpload);
 
-		const bool bValidFullDumpForCopy = Filename == FGenericCrashContext::UEMinidumpName &&
-			(FPrimaryCrashProperties::Get()->CrashDumpMode == ECrashDumpMode::FullDump || FPrimaryCrashProperties::Get()->CrashDumpMode == ECrashDumpMode::FullDumpAlways) &&
-			FPrimaryCrashProperties::Get()->CrashVersion >= ECrashDescVersions::VER_3_CrashContext &&
-			!FullCrashDumpLocation.IsEmpty();
+		const bool bIsFullDumpFile = Filename == FGenericCrashContext::UEMinidumpName && bIsFullDumpCrash;
+		const bool bValidFullDumpForCopy = bIsFullDumpFile && !FullCrashDumpLocation.IsEmpty();
 
-		if (bValidFullDumpForCopy)
+		if (bIsFullDumpFile)
 		{
-			const FString DestinationPath = FullCrashDumpLocation / FGenericCrashContext::UEMinidumpName;
-			const bool bCreated = IFileManager::Get().MakeDirectory(*FullCrashDumpLocation, true);
-			if (!bCreated)
+			if (bValidFullDumpForCopy)
 			{
-				UE_LOG(CrashReportCoreLog, Error, TEXT("Couldn't create directory for full crash dump %s"), *DestinationPath);
-			}
-			else
-			{
-				UE_LOG(CrashReportCoreLog, Warning, TEXT("Copying full crash minidump to %s"), *DestinationPath);
-				IFileManager::Get().Copy(*DestinationPath, *PathOfFileToUpload, false);
+				const FString DestinationPath = FullCrashDumpLocation / FGenericCrashContext::UEMinidumpName;
+				const bool bCreated = IFileManager::Get().MakeDirectory(*FullCrashDumpLocation, true);
+				if (!bCreated)
+				{
+					UE_LOG(CrashReportCoreLog, Error, TEXT("Couldn't create directory for full crash dump %s"), *DestinationPath);
+				}
+				else
+				{
+					UE_LOG(CrashReportCoreLog, Warning, TEXT("Copying full crash minidump to %s"), *DestinationPath);
+					IFileManager::Get().Copy(*DestinationPath, *PathOfFileToUpload, false);
+				}
 			}
 
+			UE_LOG(CrashReportCoreLog, Log, TEXT("Skipping upload of full crash dump"));
 			continue;
 		}
 
