Add support for project overrides in crash report client

Introduces a few changes that enables CrashReportClient to read project it's settings from project configuration.
* Add compile time fallback of data router to CrashReportClientEditor.
* Refactored how some global defines are set in CrashReportClient/CrashReportClientEditor target to allow different settings for development time vs retail time.
* Added code that allows project configuration files to override any CRC setting. Previously it would only apply to some properties.
* Refactored CrashReportCoreConfig.
* Added compile time defines for company name. UX changes in other CL.
* Limit core usage for CRC to 5 cores. Significantly reduces memory usage for machines with many cores.

diff --git a/Engine/Source/Programs/CrashReportClient/CrashReportClient.Target.cs b/Engine/Source/Programs/CrashReportClient/CrashReportClient.Target.cs
--- a/Engine/Source/Programs/CrashReportClient/CrashReportClient.Target.cs
+++ b/Engine/Source/Programs/CrashReportClient/CrashReportClient.Target.cs
@@ -1,14 +1,18 @@
 // Copyright Epic Games, Inc. All Rights Reserved.
 
-using System;
-using System.Collections.Generic;
-using System.Linq;
 using UnrealBuildTool;
+using System.Collections.Generic;
 
 [SupportedPlatforms("Win64", "Mac", "Linux", "LinuxArm64")]
 [SupportedConfigurations(UnrealTargetConfiguration.Debug, UnrealTargetConfiguration.Development, UnrealTargetConfiguration.Shipping)]
 public class CrashReportClientTarget : TargetRules
 {
+	[ConfigFile(ConfigHierarchyType.Engine, "CrashReportClientBuildSettings", "DataRouterFallback")]
+	public string DataRouterFallback;
+		
+	[ConfigFile(ConfigHierarchyType.Engine, "CrashReportClientBuildSettings", "CompanyName")]
+	public string CompanyName;
+	
 	[ConfigFile(ConfigHierarchyType.Engine, "CrashReportClientBuildSettings", "TelemetryUrl")]
 	public string TelemetryUrl;
 
@@ -17,8 +21,11 @@ public class CrashReportClientTarget : TargetRules
 
 	[ConfigFile(ConfigHierarchyType.Engine, "CrashReportClientBuildSettings", "TelemetryKey_Release")]
 	public string TelemetryKey_Release;
+	
+	public CrashReportClientTarget(TargetInfo Target) : this(Target, true)
+	{}
 
-	public CrashReportClientTarget(TargetInfo Target) : base(Target)
+	protected CrashReportClientTarget(TargetInfo Target, bool bSetConfiguredDefinitions) : base(Target)
 	{
 		Type = TargetType.Program;
 		LinkType = TargetLinkType.Monolithic;
@@ -53,22 +60,53 @@ public class CrashReportClientTarget : TargetRules
 		// Need to disable the bundled version of dbghelp so that CrashDebugHelper can load dbgeng.dll.
 		WindowsPlatform.bUseBundledDbgHelp = false;
 
-		// Add the definitions from config files
-		if(!string.IsNullOrWhiteSpace(TelemetryUrl))
-		{
-			AddConfigMacro("CRC_TELEMETRY_URL=", string.Format("\"{0}\"", TelemetryUrl));
-		}
-		AddConfigMacro("CRC_TELEMETRY_KEY_DEV=", string.Format("\"{0}\"", TelemetryKey_Dev));
-		AddConfigMacro("CRC_TELEMETRY_KEY_RELEASE=", string.Format("\"{0}\"", TelemetryKey_Release));
+		// Set the maximum number of cores used
+		GlobalDefinitions.Add("UE_TASKGRAPH_THREAD_LIMIT=5");
 
 		GlobalDefinitions.Add("NOINITCRASHREPORTER=1");
+
+		// Since we can't use virtual calls in the constructor allow, any inheriting type to opt out
+		// of setting these configured definitions
+		if (bSetConfiguredDefinitions)
+		{
+			GlobalDefinitions.AddRange(SetupConfiguredDefines(
+				 DataRouterFallback, CompanyName, TelemetryUrl, TelemetryKey_Dev, TelemetryKey_Release));
+		}
 	}
 
-	void AddConfigMacro(string Prefix, string Value)
+	protected static List<string> SetupConfiguredDefines(
+		string DataRouterFallback, 
+		string CompanyName, 
+		string TelemetryUrl, 
+		string TelemetryKeyDev,
+		string TelemetryKeyRelease)
 	{
-		if (!string.IsNullOrEmpty(Value) && !GlobalDefinitions.Any(x => x.StartsWith(Prefix, StringComparison.Ordinal)))
+		var Definitions = new List<string>();
+		if (!string.IsNullOrEmpty(DataRouterFallback))
+		{
+			Definitions.Add($"CRC_DATAROUTER_FALLBACK=\"{DataRouterFallback}\"");
+		}
+
+		if (!string.IsNullOrEmpty(CompanyName))
 		{
-			GlobalDefinitions.Add(Prefix + Value);
+			Definitions.Add($"CRC_COMPANY_NAME_FALLBACK=\"{CompanyName}\"");
 		}
+
+		if(!string.IsNullOrWhiteSpace(TelemetryUrl))
+		{
+			Definitions.Add($"CRC_TELEMETRY_URL=\"{TelemetryUrl}\"");
+		}
+
+		if (!string.IsNullOrWhiteSpace(TelemetryKeyDev))
+		{
+			Definitions.Add($"CRC_TELEMETRY_KEY_DEV=\"{TelemetryKeyDev}\"");
+		}
+
+		if (!string.IsNullOrWhiteSpace(TelemetryKeyRelease))
+		{
+			Definitions.Add($"CRC_TELEMETRY_KEY_RELEASE=\"{TelemetryKeyRelease}\"");
+		}
+
+		return Definitions;
 	}
 }
diff --git a/Engine/Source/Programs/CrashReportClient/CrashReportClientEditor.Target.cs b/Engine/Source/Programs/CrashReportClient/CrashReportClientEditor.Target.cs
--- a/Engine/Source/Programs/CrashReportClient/CrashReportClientEditor.Target.cs
+++ b/Engine/Source/Programs/CrashReportClient/CrashReportClientEditor.Target.cs
@@ -1,13 +1,30 @@
 // Copyright Epic Games, Inc. All Rights Reserved.
 
 using UnrealBuildTool;
-using System.Collections.Generic;
 
 [SupportedPlatforms("Win64", "Mac", "Linux")]
 [SupportedConfigurations(UnrealTargetConfiguration.Debug, UnrealTargetConfiguration.Development, UnrealTargetConfiguration.Shipping)]
-public class CrashReportClientEditorTarget : CrashReportClientTarget
+public sealed class CrashReportClientEditorTarget : CrashReportClientTarget
 {
-	public CrashReportClientEditorTarget(TargetInfo Target) : base(Target)
+	// Override the configuration values from CrashReportClient with these using another
+	// configuration block: [CrashReportClientEditorBuildSettings]
+	
+	[ConfigFile(ConfigHierarchyType.Engine, "CrashReportClientEditorBuildSettings", "DataRouterFallback")]
+	public new string DataRouterFallback;
+		
+	[ConfigFile(ConfigHierarchyType.Engine, "CrashReportClientEditorBuildSettings", "CompanyName")]
+	public new string CompanyName;
+	
+	[ConfigFile(ConfigHierarchyType.Engine, "CrashReportClientEditorBuildSettings", "TelemetryUrl")]
+	public new string TelemetryUrl;
+
+	[ConfigFile(ConfigHierarchyType.Engine, "CrashReportClientEditorBuildSettings", "TelemetryKey_Dev")]
+	public new string TelemetryKey_Dev;
+
+	[ConfigFile(ConfigHierarchyType.Engine, "CrashReportClientEditorBuildSettings", "TelemetryKey_Release")]
+	public new string TelemetryKey_Release;
+	
+	public CrashReportClientEditorTarget(TargetInfo Target) : base(Target, false /* bSetConfiguredDefinitions */)
 	{
 		LaunchModuleName = "CrashReportClientEditor";
 
@@ -19,6 +36,7 @@ public class CrashReportClientEditorTarget : CrashReportClientTarget
 
 		if (bHostRecoverySvc)
 		{
+			throw new BuildException("No longer supported.");
 			AdditionalPlugins.Add("UdpMessaging");
 			AdditionalPlugins.Add("ConcertSyncServer");
 			bCompileWithPluginSupport = true; // Enable Developer plugins (like Concert!)
@@ -30,5 +48,10 @@ public class CrashReportClientEditorTarget : CrashReportClientTarget
 				GlobalDefinitions.Add("PLATFORM_SUPPORTS_MESSAGEBUS=1");
 			}
 		}
+		
+		// We can now set the configured definitions from CrashReportClientEditorBuildSettings section
+		GlobalDefinitions.AddRange(SetupConfiguredDefines(
+			DataRouterFallback, CompanyName, TelemetryUrl, TelemetryKey_Dev, TelemetryKey_Release));
 	}
+	
 }
diff --git a/Engine/Source/Programs/CrashReportClient/Private/CrashReportClient.cpp b/Engine/Source/Programs/CrashReportClient/Private/CrashReportClient.cpp
--- a/Engine/Source/Programs/CrashReportClient/Private/CrashReportClient.cpp
+++ b/Engine/Source/Programs/CrashReportClient/Private/CrashReportClient.cpp
@@ -18,9 +18,16 @@ struct FCrashReportUtil
 	/** Formats processed diagnostic text by adding additional information about machine and user. */
 	static FText FormatDiagnosticText( const FText& DiagnosticText )
 	{
-		const FString LoginId = FPrimaryCrashProperties::Get()->LoginId.AsString();
-		const FString EpicAccountId = FPrimaryCrashProperties::Get()->EpicAccountId.AsString();
-		return FText::Format( LOCTEXT( "CrashReportClientCallstackPattern", "LoginId:{0}\nEpicAccountId:{1}\n\n{2}" ), FText::FromString( LoginId ), FText::FromString( EpicAccountId ), DiagnosticText );
+		TStringBuilder<512> Accounts;
+		if (const FString LoginId = FPrimaryCrashProperties::Get()->LoginId.AsString(); !LoginId.IsEmpty())
+		{
+			Accounts.Appendf(TEXT("LoginId:%s\n"), *LoginId);
+		}
+		if (const FString EpicAccountId= FPrimaryCrashProperties::Get()->EpicAccountId.AsString(); !EpicAccountId.IsEmpty())
+		{
+			Accounts.Appendf(TEXT("EpicAccountId:%s\n"), *EpicAccountId);
+		}
+		return FText::Format(LOCTEXT("CrashReportClientCallstackPattern", "{0}\n{1}"), FText::FromString(Accounts.ToString()), DiagnosticText);
 	}
 };
 
diff --git a/Engine/Source/Programs/CrashReportClient/Private/CrashReportClientApp.cpp b/Engine/Source/Programs/CrashReportClient/Private/CrashReportClientApp.cpp
--- a/Engine/Source/Programs/CrashReportClient/Private/CrashReportClientApp.cpp
+++ b/Engine/Source/Programs/CrashReportClient/Private/CrashReportClientApp.cpp
@@ -1094,6 +1094,15 @@ void RunCrashReportClient(const TCHAR* CommandLine)
 
 		// Load error report generated by the process from disk
 		FPlatformErrorReport ErrorReport = LoadErrorReport();
+
+		// Apply project and crash overrides
+		FString CrashSettingsIni;
+		ErrorReport.FindFirstReportFileWithExtension(CrashSettingsIni, FGenericCrashContext::CrashConfigExtension);
+		FCrashReportCoreConfig::Get().ApplyProjectOverrides(FPaths::Combine(ErrorReport.GetReportDirectory(), CrashSettingsIni));
+
+		// At this point all overrides has been applied
+		FCrashReportCoreConfig::Get().PrintSettingsToLog();
+		
 		SendErrorReport(ErrorReport, FApp::IsUnattended(), bImplicitSendFromCmd);
 
 		if (AnalyticsEnabledFromCmd)
@@ -1154,6 +1163,14 @@ void RunCrashReportClient(const TCHAR* CommandLine)
 					// Build error report in memory.
 					bool bCrashedThreadCallstackAvailable = false;
 					FPlatformErrorReport ErrorReport = CollectErrorReport(RecoveryServicePtr.Get(), MonitorPid, CrashContext, MonitorWritePipe, bCrashedThreadCallstackAvailable);
+					
+					// Apply project and crash overrides
+					FString CrashSettingsIni;
+					ErrorReport.FindFirstReportFileWithExtension(CrashSettingsIni, FGenericCrashContext::CrashConfigExtension);
+					FCrashReportCoreConfig::Get().ApplyProjectOverrides(FPaths::Combine(ErrorReport.GetReportDirectory(), CrashSettingsIni));
+
+					// At this point all overrides has been applied
+					FCrashReportCoreConfig::Get().PrintSettingsToLog();
 
 					// Log cases where the PCallstack is missing. For analytics, this event hints that the remote app exited before CRC could walk the process stack and possibly means
 					// that the timeout in the crashing process waiting for CRC to reply 'continue' is too short.
diff --git a/Engine/Source/Runtime/CrashReportCore/Private/CrashReportCoreConfig.cpp b/Engine/Source/Runtime/CrashReportCore/Private/CrashReportCoreConfig.cpp
--- a/Engine/Source/Runtime/CrashReportCore/Private/CrashReportCoreConfig.cpp
+++ b/Engine/Source/Runtime/CrashReportCore/Private/CrashReportCoreConfig.cpp
@@ -6,7 +6,6 @@
 #include "Misc/ConfigCacheIni.h"
 #include "Misc/App.h"
 #include "Misc/EngineBuildSettings.h"
-#include "GenericPlatform/GenericPlatformCrashContext.h"
 
 FCrashReportCoreConfig::FCrashReportCoreConfig()
 	: DiagnosticsFilename( TEXT( "Diagnostics.txt" ) )
@@ -23,41 +22,27 @@ FCrashReportCoreConfig::FCrashReportCoreConfig()
 	{
 		CrashReportClientVersion = TEXT("0.0.0");
 	}
-	UE_LOG(CrashReportCoreLog, Log, TEXT("CrashReportClientVersion=%s"), *CrashReportClientVersion);
 
 	if (!GConfig->GetString( *SectionName, TEXT( "CrashReportReceiverIP" ), CrashReportReceiverIP, GEngineIni ))
 	{
 		// Use the default value (blank/disabled)
 		CrashReportReceiverIP = TEXT("");
 	}
-	if (CrashReportReceiverIP.IsEmpty())
-	{
-		UE_LOG(CrashReportCoreLog, Log, TEXT("CrashReportReceiver disabled"));
-	}
-	else
-	{
-		UE_LOG(CrashReportCoreLog, Log, TEXT("CrashReportReceiverIP: %s"), *CrashReportReceiverIP);
-	}
 
 	if (!GConfig->GetString(*SectionName, TEXT("DataRouterUrl"), DataRouterUrl, GEngineIni))
 	{
-		// Use the default value.
+#if defined CRC_DATAROUTER_FALLBACK
+		DataRouterUrl = TEXT(CRC_DATAROUTER_FALLBACK);
+#else
 		DataRouterUrl = TEXT("");
-	}
-	if (DataRouterUrl.IsEmpty())
-	{
-		UE_LOG(CrashReportCoreLog, Log, TEXT("DataRouter disabled"));
-	}
-	else
-	{
-		UE_LOG(CrashReportCoreLog, Log, TEXT("DataRouterUrl: %s"), *DataRouterUrl);
+#endif
 	}
 
 	if (FEngineBuildSettings::IsInternalBuild())
 	{
 		bAllowToBeContacted = true;
 	}
-	else if (!GConfig->GetBool( TEXT( "CrashReportClient" ), TEXT( "bAllowToBeContacted" ), bAllowToBeContacted, GEngineIni ))
+	else if (!GConfig->GetBool(*SectionName, TEXT("bAllowToBeContacted"), bAllowToBeContacted, GEngineIni ))
 	{
 		// Default to true when unattended when config is missing. This is mostly for dedicated servers that do not have config files for CRC.
 		if (bUnattended)
@@ -66,7 +51,7 @@ FCrashReportCoreConfig::FCrashReportCoreConfig()
 		}
 	}
 
-	if (!GConfig->GetBool( TEXT( "CrashReportClient" ), TEXT( "bSendLogFile" ), bSendLogFile, GEngineIni ))
+	if (!GConfig->GetBool(*SectionName, TEXT("bSendLogFile"), bSendLogFile, GEngineIni ))
 	{
 		// Default to true when unattended when config is missing. This is mostly for dedicated servers that do not have config files for CRC.
 		if (bUnattended)
@@ -75,13 +60,47 @@ FCrashReportCoreConfig::FCrashReportCoreConfig()
 		}
 	}
 
-	if (!GConfig->GetInt(TEXT("CrashReportClient"), TEXT("UserCommentSizeLimit"), UserCommentSizeLimit, GEngineIni))
+	if (!GConfig->GetInt(*SectionName, TEXT("UserCommentSizeLimit"), UserCommentSizeLimit, GEngineIni))
 	{
 		UserCommentSizeLimit = 4000;
 	}
 
-	FConfigFile EmptyConfigFile;
-	SetProjectConfigOverrides(EmptyConfigFile);
+	if (!GConfig->GetBool(*SectionName, TEXT("bHideLogFilesOption"), bHideLogFilesOption, GEngineIni))
+	{
+		bHideLogFilesOption = false;
+	}
+
+	if (!GConfig->GetBool(*SectionName, TEXT("bHideRestartOption"), bHideRestartOption, GEngineIni))
+	{
+		bHideRestartOption = false;
+	}
+	
+	if (!GConfig->GetBool(*SectionName, TEXT("bIsAllowedToCloseWithoutSending"), bIsAllowedToCloseWithoutSending, GEngineIni))
+	{
+		// Default to false (show the option) when config is missing.
+		bIsAllowedToCloseWithoutSending = true;
+	}
+	
+	if (!GConfig->GetBool(*SectionName, TEXT("bShowEndpointInTooltip"), bShowEndpointInTooltip, GEngineIni))
+	{
+		bShowEndpointInTooltip = false;
+	}
+
+	if (!GConfig->GetString(*SectionName, TEXT("CompanyName"), CompanyName, GEngineIni))
+	{
+#if defined(CRC_COMPANY_NAME_FALLBACK)
+		CompanyName = TEXT(CRC_COMPANY_NAME_FALLBACK);
+#else
+		CompanyName = TEXT("");
+#endif
+	}
+
+#if PLATFORM_WINDOWS
+	if (!GConfig->GetBool(*SectionName, TEXT("bIsAllowedToCopyFilesToClipboard"), bIsAllowedToCopyFilesToClipboard, GEngineIni))
+	{
+		bIsAllowedToCopyFilesToClipboard = false;
+	}
+#endif
 
 	ReadFullCrashDumpConfigurations();
 }
@@ -98,49 +117,36 @@ void FCrashReportCoreConfig::SetSendLogFile( bool bNewValue )
 	GConfig->SetBool( *SectionName, TEXT( "bSendLogFile" ), bSendLogFile, GEngineIni );
 }
 
-void FCrashReportCoreConfig::SetProjectConfigOverrides(const FConfigFile& InConfigFile)
+void FCrashReportCoreConfig::ApplyProjectOverrides(const FString& ConfigFilePath)
 {
-	const FConfigSection* Section = InConfigFile.FindSection(FGenericCrashContext::ConfigSectionName);
-
-	// Default to false (show the option) when config is missing.
-	bHideLogFilesOption = false;
-	bHideRestartOption = false;
-
-	// Default to true (Allow the user to close without sending) when config is missing.
-	bIsAllowedToCloseWithoutSending = true;
-
-	// Try to read values from override config file
-	if (Section != nullptr)
+	if (FPaths::FileExists(ConfigFilePath))
 	{
-		const FConfigValue* HideLogFilesOptionValue = Section->Find(TEXT("bHideLogFilesOption"));
-		if (HideLogFilesOptionValue != nullptr)
-		{
-			bHideLogFilesOption = FCString::ToBool(*HideLogFilesOptionValue->GetValue());
-		}
-
-		const FConfigValue* HideRestartOptionValue = Section->Find(TEXT("bHideRestartOption"));
-		if (HideRestartOptionValue != nullptr)
-		{
-			bHideRestartOption = FCString::ToBool(*HideRestartOptionValue->GetValue());
-		}
-
-		const FConfigValue* IsAllowedToCloseWithoutSendingValue = Section->Find(TEXT("bIsAllowedToCloseWithoutSending"));
-		if (IsAllowedToCloseWithoutSendingValue != nullptr)
-		{
-			bIsAllowedToCloseWithoutSending = FCString::ToBool(*IsAllowedToCloseWithoutSendingValue->GetValue());
-		}
+		UE_LOG(CrashReportCoreLog, Display, TEXT("Applying project settings from '%s'."), *ConfigFilePath)
+		FConfigFile ProjectConfig;
+		ProjectConfig.Read(ConfigFilePath);
+		SetProjectConfigOverrides(ProjectConfig);
+	}
+}
 
+void FCrashReportCoreConfig::SetProjectConfigOverrides(const FConfigFile& InConfigFile)
+{
+	InConfigFile.GetString(*SectionName, TEXT("CrashReportClientVersion"), CrashReportClientVersion);
+	InConfigFile.GetString(*SectionName, TEXT("CrashReportReceiverIP"), CrashReportReceiverIP);
+	InConfigFile.GetString(*SectionName, TEXT("DataRouterUrl"), DataRouterUrl);
+	InConfigFile.GetBool(*SectionName, TEXT("bAllowToBeContacted"), bAllowToBeContacted);
+	InConfigFile.GetBool(*SectionName, TEXT( "bSendLogFile" ), bSendLogFile);
+	InConfigFile.GetInt(*SectionName, TEXT("UserCommentSizeLimit"), UserCommentSizeLimit);
+	InConfigFile.GetBool(*SectionName, TEXT("bHideLogFilesOption"), bHideLogFilesOption);
+	InConfigFile.GetBool(*SectionName, TEXT("bHideRestartOption"), bHideRestartOption);
+	InConfigFile.GetBool(*SectionName, TEXT("bIsAllowedToCloseWithoutSending"), bIsAllowedToCloseWithoutSending);
+	InConfigFile.GetBool(*SectionName, TEXT("bShowEndpointInTooltip"), bShowEndpointInTooltip);
+	InConfigFile.GetString(*SectionName, TEXT("CompanyName"), CompanyName);
 #if PLATFORM_WINDOWS
-		const FConfigValue* IsAllowedToCopyFilesToClipboard = Section->Find(TEXT("bIsAllowedToCopyFilesToClipboard"));
-		if (IsAllowedToCopyFilesToClipboard  != nullptr)
-		{
-			bIsAllowedToCopyFilesToClipboard  = FCString::ToBool(*IsAllowedToCopyFilesToClipboard->GetValue());
-		}
+	InConfigFile.GetBool(*SectionName, TEXT("bIsAllowedToCopyFilesToClipboard"), bIsAllowedToCopyFilesToClipboard);
 #endif
-	}
 }
 
-const FString FCrashReportCoreConfig::GetFullCrashDumpLocationForBranch( const FString& BranchName ) const
+FString FCrashReportCoreConfig::GetFullCrashDumpLocationForBranch(const FString& BranchName) const
 {
 	for (const auto& It : FullCrashDumpConfigurations)
 	{
@@ -159,32 +165,55 @@ const FString FCrashReportCoreConfig::GetFullCrashDumpLocationForBranch( const F
 	return TEXT( "" );
 }
 
-FString FCrashReportCoreConfig::GetKey( const FString& KeyName )
-{
-	FString Result;
-	if (!GConfig->GetString( *SectionName, *KeyName, Result, GEngineIni ))
-	{
-		return TEXT( "" );
-	}
-	return Result;
-}
-
 void FCrashReportCoreConfig::ReadFullCrashDumpConfigurations()
 {
+	auto GetKey = [this](const TCHAR* KeyName) -> FString
+	{
+		FString Result;
+		if (!GConfig->GetString( *SectionName, KeyName, Result, GEngineIni ))
+		{
+			return TEXT( "" );
+		}
+		return MoveTemp(Result);
+	};
+	
 	for (int32 NumEntries = 0;; ++NumEntries)
 	{
-		FString Branch = GetKey( FString::Printf( TEXT( "FullCrashDump_%i_Branch" ), NumEntries ) );
+		FString Branch = GetKey( *FString::Printf( TEXT( "FullCrashDump_%i_Branch" ), NumEntries ) );
 		if (Branch.IsEmpty())
 		{
 			break;
 		}
 
-		const FString NetworkLocation = GetKey( FString::Printf( TEXT( "FullCrashDump_%i_Location" ), NumEntries ) );
+		const FString NetworkLocation = GetKey( *FString::Printf( TEXT( "FullCrashDump_%i_Location" ), NumEntries ) );
 		const bool bExactMatch = !Branch.EndsWith( TEXT( "*" ) );
 		Branch.ReplaceInline( TEXT( "*" ), TEXT( "" ) );
 
 		FullCrashDumpConfigurations.Add( FFullCrashDumpEntry( Branch, NetworkLocation, bExactMatch ) );
+	}
+}
 
-		UE_LOG( CrashReportCoreLog, Log, TEXT( "FullCrashDump: %s, NetworkLocation: %s, bExactMatch:%i" ), *Branch, *NetworkLocation, bExactMatch );
+void FCrashReportCoreConfig::PrintSettingsToLog() const
+{
+	UE_LOG(CrashReportCoreLog, Log, TEXT("CrashReportClientVersion=%s"), *CrashReportClientVersion);
+	if (CrashReportReceiverIP.IsEmpty())
+	{
+		UE_LOG(CrashReportCoreLog, Log, TEXT("CrashReportReceiver disabled"));
+	}
+	else
+	{
+		UE_LOG(CrashReportCoreLog, Log, TEXT("CrashReportReceiverIP: %s"), *CrashReportReceiverIP);
+	}
+	if (DataRouterUrl.IsEmpty())
+	{
+		UE_LOG(CrashReportCoreLog, Log, TEXT("DataRouter disabled"));
+	}
+	else
+	{
+		UE_LOG(CrashReportCoreLog, Log, TEXT("DataRouterUrl: %s"), *DataRouterUrl);
+	}
+	for (auto Location : FullCrashDumpConfigurations)
+	{
+		UE_LOG( CrashReportCoreLog, Log, TEXT( "FullCrashDump: %s, NetworkLocation: %s, bExactMatch:%i" ), *Location.BranchName, *Location.Location, Location.bExactMatch);
 	}
 }
diff --git a/Engine/Source/Runtime/CrashReportCore/Public/CrashReportCoreConfig.h b/Engine/Source/Runtime/CrashReportCore/Public/CrashReportCoreConfig.h
--- a/Engine/Source/Runtime/CrashReportCore/Public/CrashReportCoreConfig.h
+++ b/Engine/Source/Runtime/CrashReportCore/Public/CrashReportCoreConfig.h
@@ -80,33 +80,43 @@ struct FCrashReportCoreConfig
 		return DiagnosticsFilename;
 	}
 
-	const bool& GetAllowToBeContacted() const
+	bool GetAllowToBeContacted() const
 	{
 		return bAllowToBeContacted;
 	}
 
-	const bool& GetSendLogFile() const
+	bool GetSendLogFile() const
 	{
 		return bSendLogFile;
 	}
 
-	const bool& GetHideLogFilesOption() const
+	bool GetHideLogFilesOption() const
 	{
 		return bHideLogFilesOption;
 	}
 
-	const bool& GetHideRestartOption() const
+	bool GetHideRestartOption() const
 	{
 		return bHideRestartOption;
 	}
 
-	const bool& IsAllowedToCloseWithoutSending() const
+	bool IsAllowedToCloseWithoutSending() const
 	{
 		return bIsAllowedToCloseWithoutSending;
 	}
 
+	bool GetShowEndpointInTooltip() const
+	{
+		return bShowEndpointInTooltip;
+	}
+
+	const FString& GetCompanyName() const
+	{
+		return CompanyName;
+	}
+
 #if PLATFORM_WINDOWS
-	const bool IsAllowedToCopyFilesToClipboard() const
+	bool IsAllowedToCopyFilesToClipboard() const
 	{
 		return bIsAllowedToCopyFilesToClipboard;
 	}
@@ -120,18 +130,23 @@ struct FCrashReportCoreConfig
 	void SetAllowToBeContacted( bool bNewValue );
 	void SetSendLogFile( bool bNewValue );
 
+
+	void ApplyProjectOverrides(const FString& ConfigFilePath);
+
 	/** Set config values that are determined by the crashing application saving a config file to the crash folder */
 	void SetProjectConfigOverrides(const FConfigFile& InConfigFile);
 
 	/**
 	 * @return location for full crash dump for the specified branch.
 	 */
-	const FString GetFullCrashDumpLocationForBranch( const FString& BranchName ) const;
+	FString GetFullCrashDumpLocationForBranch(const FString& BranchName) const;
 
-protected:
-	/** Returns empty string if couldn't read */
-	FString GetKey( const FString& KeyName );
+	/**
+	 * Prints the current settings to the log. Call once all applicable overrides has been applied.
+	 */
+	void PrintSettingsToLog() const;
 
+protected:
 	/** Reads FFullCrashDumpEntry config entries. */
 	void ReadFullCrashDumpConfigurations();
 
@@ -172,6 +187,12 @@ protected:
 	/** Whether the user is allowed to close the crash reporter without sending a report */
 	bool bIsAllowedToCloseWithoutSending;
 
+	/** Optional string displayed in the "allow to be contacted" text. */
+	FString CompanyName;
+
+	/** Whether to show the endpoint (reciever or datarouter url) */
+	bool bShowEndpointInTooltip;
+
 #if PLATFORM_WINDOWS
 	/** Whether the user is allowed to copy the report files to the clipboard */
 	bool bIsAllowedToCopyFilesToClipboard;
