CIS fix.

diff --git a/Engine/Source/Runtime/CrashReportCore/Private/Mac/MacErrorReport.cpp b/Engine/Source/Runtime/CrashReportCore/Private/Mac/MacErrorReport.cpp
--- a/Engine/Source/Runtime/CrashReportCore/Private/Mac/MacErrorReport.cpp
+++ b/Engine/Source/Runtime/CrashReportCore/Private/Mac/MacErrorReport.cpp
@@ -60,11 +60,11 @@ FString FMacErrorReport::FindCrashedAppPath() const
 			int32 LineEnd = -1;
 			if(PathData.FindChar( TCHAR('\r'), LineEnd ))
 			{
-				PathData.LeftInline(LineEnd, false);
+				PathData.LeftInline(LineEnd, EAllowShrinking::No);
 			}
 			if(PathData.FindChar( TCHAR('\n'), LineEnd ))
 			{
-				PathData.LeftInline(LineEnd, false);
+				PathData.LeftInline(LineEnd, EAllowShrinking::No);
 			}
 			return PathData;
 		}
