Added additional logging to Crash Uploads to help identify issues with reports not being uploaded because they are too large

diff --git a/Engine/Source/Runtime/CrashReportCore/Private/CrashUpload.cpp b/Engine/Source/Runtime/CrashReportCore/Private/CrashUpload.cpp
--- a/Engine/Source/Runtime/CrashReportCore/Private/CrashUpload.cpp
+++ b/Engine/Source/Runtime/CrashReportCore/Private/CrashUpload.cpp
@@ -494,7 +494,7 @@ void FCrashUploadToReceiver::CompressAndSendData()
 	Request->SetHeader(TEXT("CompressedSize"), TTypeToString<int32>::ToString(CompressedData.CompressedSize) );
 	Request->SetHeader(TEXT("UncompressedSize"), TTypeToString<int32>::ToString(CompressedData.UncompressedSize) );
 	Request->SetHeader(TEXT("NumberOfFiles"), TTypeToString<int32>::ToString(CompressedData.FileCount) );
-	UE_LOG( CrashReportCoreLog, Log, TEXT( "Sending HTTP request: %s" ), *Request->GetURL() );
+	UE_LOG( CrashReportCoreLog, Log, TEXT( "Sending HTTP request: %s, Payload size: %d" ), *Request->GetURL(), CompressedData.Data.Num());
 
 	if (Request->ProcessRequest())
 	{
@@ -534,7 +534,7 @@ void FCrashUploadToReceiver::PostReportComplete()
 	Request->SetURL(UrlPrefix / TEXT("UploadComplete"));
 	Request->SetHeader( TEXT( "Content-Type" ), TEXT( "text/plain; charset=us-ascii" ) );
 	Request->SetContent(PostData);
-	UE_LOG( CrashReportCoreLog, Log, TEXT( "Sending HTTP request: %s" ), *Request->GetURL() );
+	UE_LOG( CrashReportCoreLog, Log, TEXT( "Sending HTTP request: %s, Payload size: %d" ), *Request->GetURL(), PostData.Num());
 
 	if (Request->ProcessRequest())
 	{
@@ -789,7 +789,7 @@ void FCrashUploadToDataRouter::CompressAndSendData()
 	Request->SetHeader(TEXT("Content-Type"), TEXT("application/octet-stream"));
 	Request->SetURL(DataRouterUrl + UrlParams);
 	Request->SetContent(CompressedData.Data);
-	UE_LOG(CrashReportCoreLog, Log, TEXT("Sending HTTP request: %s"), *Request->GetURL());
+	UE_LOG(CrashReportCoreLog, Log, TEXT("Sending HTTP request: %s, Payload size: %d"), *Request->GetURL(), CompressedData.Data.Num());
 
 	if (Request->ProcessRequest())
 	{
