Only send CrashReportClient.ReportCrashUploaded telemetry event if upload succeeds.

diff --git a/Engine/Source/Runtime/CrashReportCore/Private/CrashDescription.cpp b/Engine/Source/Runtime/CrashReportCore/Private/CrashDescription.cpp
--- a/Engine/Source/Runtime/CrashReportCore/Private/CrashDescription.cpp
+++ b/Engine/Source/Runtime/CrashReportCore/Private/CrashDescription.cpp
@@ -319,30 +319,12 @@ void FPrimaryCrashProperties::SendPreUploadAnalytics()
   * @Type Client
   * @Owner Wes.Hunt
   *
-  * @EventParam bHasPrimaryData - Whether the crash loaded data successfully from a crash context or legacy metadata file that was saved by the crashed process ("true" or "false")
-  * @EventParam CrashVersion" - Describes the version of the crash data pipeline we used on the client side (1 = Legacy metadata based, 2 = unused, 3 = New crash context based)
-  * @EventParam CrashGUID - GUID for this event generated by the crashed process (e.g. Windows format is UECC-Windows-1F07494140C2669B52AC96A1C4D20F83_0000 with the last four digits representing the index of event within a single session)
-  * @EventParam PlatformCallbackResult - Platform-specific UE Core value. (integer)
-  * @EventParam CrashType - The name of the crash report type (e.g Assert, Ensure, Hang)
-  * @EventParam GameName - The name of the game that crashed from FApp::GetProjectName(). (e.g. UE-QAGame)
-  * @EventParam EngineVersion - The engine version of the crashed process from FEngineVersion::Current(). (e.g. 4.13.1-3142249+++UE+Release-4.13)
-  * @EventParam CrashReportClientVersion - CrashReportClient version, set in its config. (e.g. "1.0.0")
-  * @EventParam LoginID - Unique string associated with the login account of the current machine from FPlatformMisc::GetLoginId().
-  * @EventParam UserName - Username of the user of the process that crashed from FPlatformProcess::UserName().
-  * @EventParam EpicAccountId - Epic account ID for the user who last used the Launcher on this machine.
-  * @EventParam Platform - Detailed platform name string. (e.g. "Win64 [Windows 7 Service Pack 1 64b]")
-  * @EventParam TimeOfCrash - Time in UTC that the crash was read by the Crash Report Client in raw ticks (e.g. "636153351834530000", ticks are 100-nanosecond intervals that have elapsed since 12:00:00 midnight, January 1, 0001)
-  * @EventParam EngineMode - The type of process that crashed ("Commandlet", "Game", "Editor" or "Server")
-  * @EventParam EngineModeEx - The "vanilla" status of the engine, whether it's a totally standard, unmodified Epic-distributed build. ("Unset", "Vanilla" or "Dirty")
-  * @EventParam AppDefaultLocale - The ICU default locale string or "en" if ICU is not enabled.
-  * @EventParam UserActivityHint - Application-specific user activity string, if set in the crashed process. The meaning is game/app-specific.
-  * @EventParam GameSessionID - Application-specific session Id, if set in the crashed process.
-  * @EventParam DeploymentName - Deployment name, also known as EpicApp. (e.g. "DevPlaytest", "PublicTest", "Live", etc)
-   * @EventParam EngineData - Json object containing any context that had been added to the crash context using FGenericCrashContext::SetEngineData.
- * @EventParam GameData - Json object containing any context that had been added to the crash context using FGenericCrashContext::SetGameData.
-*
-  * @Comments These events should exactly match corresponding CrashReportClient.ReportEnsure events that the CRC sent before the upload started.
-  * This event will be missing if the upload failed for any reason so the difference between the event counts should tell you the success rate.
+  * @EventParam Result - If the HTTP end point could successfully be reached
+  * @EventParam ResponseCode - HTTP Response code from the end point.
+  * @EventParam Duration - Time to compress and send request.
+  * @EventParam PayloadSize - Compressed size of payload (attached files)
+  * @EventParam ReportCount - Number of reports uploaded (if more than one above values are accumulative or last observed)
+  *
   */
 void SendPostUploadEnsureAnalytics(const TArray<FAnalyticsEventAttribute>& InCrashAttributes)
 {
@@ -357,43 +339,29 @@ void SendPostUploadEnsureAnalytics(const TArray<FAnalyticsEventAttribute>& InCra
  * @Type Client
  * @Owner Wes.Hunt
  *
- * @EventParam bHasPrimaryData - Whether the crash loaded data successfully from a crash context or legacy metadata file that was saved by the crashed process ("true" or "false")
- * @EventParam CrashVersion" - Describes the version of the crash data pipeline we used on the client side (1 = Legacy metadata based, 2 = unused, 3 = New crash context based)
- * @EventParam CrashGUID - GUID for this event generated by the crashed process (e.g. Windows format is UECC-Windows-1F07494140C2669B52AC96A1C4D20F83_0000 with the last four digits representing the index of event within a single session)
- * @EventParam PlatformCallbackResult - Platform-specific UE Core value. (integer)
- * @EventParam CrashType - The name of the crash report type (e.g Assert, Ensure, Hang)
- * @EventParam GameName - The name of the game that crashed from FApp::GetProjectName(). (e.g. UE-QAGame)
- * @EventParam EngineVersion - The engine version of the crashed process from FEngineVersion::Current(). (e.g. 4.13.1-3142249+++UE+Release-4.13)
- * @EventParam CrashReportClientVersion - CrashReportClient version, set in its config. (e.g. "1.0.0")
- * @EventParam LoginID - Unique string associated with the login account of the current machine from FPlatformMisc::GetLoginId().
- * @EventParam UserName - Username of the user of the process that crashed from FPlatformProcess::UserName().
- * @EventParam EpicAccountId - Epic account ID for the user who last used the Launcher on this machine.
- * @EventParam Platform - Detailed platform name string. (e.g. "Win64 [Windows 7 Service Pack 1 64b]")
- * @EventParam TimeOfCrash - Time in UTC that the crash was read by the Crash Report Client in raw ticks (e.g. "636153351834530000", ticks are 100-nanosecond intervals that have elapsed since 12:00:00 midnight, January 1, 0001)
- * @EventParam EngineMode - The type of process that crashed ("Commandlet", "Game", "Editor" or "Server")
- * @EventParam EngineModeEx - The "vanilla" status of the engine, whether it's a totally standard, unmodified Epic-distributed build. ("Unset", "Vanilla" or "Dirty")
- * @EventParam AppDefaultLocale - The ICU default locale string or "en" if ICU is not enabled.
- * @EventParam UserActivityHint - Application-specific user activity string, if set in the crashed process. The meaning is game/app-specific.
- * @EventParam GameSessionID - Application-specific session Id, if set in the crashed process.
- * @EventParam DeploymentName - Deployment name, also known as EpicApp. (e.g. "DevPlaytest", "PublicTest", "Live", etc)
- * @EventParam EngineData - Json object containing any context that had been added to the crash context using FGenericCrashContext::SetEngineData.
- * @EventParam GameData - Json object containing any context that had been added to the crash context using FGenericCrashContext::SetGameData.
+ * @EventParam Result - If the HTTP end point could successfully be reached
+ * @EventParam ResponseCode - HTTP Response code from the end point.
+ * @EventParam Duration - Time to compress and send request.
+ * @EventParam PayloadSize - Compressed size of payload (attached files)
  *
- * @Comments These events should exactly match corresponding CrashReportClient.ReportCrash events that the CRC sent before the upload started.
- * This event will be missing if the upload failed for any reason so the difference between the event counts should tell you the success rate.
  */
 void SendPostUploadCrashAnalytics(const TArray<FAnalyticsEventAttribute>& InCrashAttributes)
 {
 	FCrashReportAnalytics::GetProvider().RecordEvent(TEXT("CrashReportClient.ReportCrashUploaded"), InCrashAttributes);
 }
 
-void FPrimaryCrashProperties::SendPostUploadAnalytics()
+void FPrimaryCrashProperties::SendPostUploadAnalytics(double Duration, bool bResult, int32 ResponseCode, uint32 PayloadSize, uint32 ReportCount)
 {
-	TArray<FAnalyticsEventAttribute> CrashAttributes;
-	MakeCrashEventAttributes(CrashAttributes);
-
 	if (FCrashReportAnalytics::IsAvailable())
 	{
+		const TArray<FAnalyticsEventAttribute> CrashAttributes = {
+			FAnalyticsEventAttribute(TEXT("Duration"), Duration),
+			FAnalyticsEventAttribute(TEXT("Result"), bResult),
+			FAnalyticsEventAttribute(TEXT("ResponseCode"), ResponseCode),
+			FAnalyticsEventAttribute(TEXT("PayloadSize"), PayloadSize),
+			FAnalyticsEventAttribute(TEXT("ReportCount"), ReportCount),
+		};
+		
 		if (bIsEnsure)
 		{
 			SendPostUploadEnsureAnalytics(CrashAttributes);
diff --git a/Engine/Source/Runtime/CrashReportCore/Private/CrashUpload.cpp b/Engine/Source/Runtime/CrashReportCore/Private/CrashUpload.cpp
--- a/Engine/Source/Runtime/CrashReportCore/Private/CrashUpload.cpp
+++ b/Engine/Source/Runtime/CrashReportCore/Private/CrashUpload.cpp
@@ -1,6 +1,8 @@
 // Copyright Epic Games, Inc. All Rights Reserved.
 	
 #include "CrashUpload.h"
+
+#include "AnalyticsEventAttribute.h"
 #include "CrashReportCoreModule.h"
 #include "HAL/FileManager.h"
 #include "Misc/Compression.h"
@@ -249,18 +251,17 @@ bool FCrashUploadBase::CompressData(const TArray<FString>& InPendingFiles, FComp
 
 	int UncompressedSize = UncompressedData.Num();
 
-	uint8* CompressedDataRaw = new uint8[UncompressedSize];
+	TArray<uint8> CompressedDataRaw;
+	CompressedDataRaw.AddUninitialized(UncompressedSize);
 
 	OutCompressedData.FileCount = CurrentFileIndex;
 	OutCompressedData.CompressedSize = UncompressedSize;
 	OutCompressedData.UncompressedSize = UncompressedSize;
-	const bool bResult = FCompression::CompressMemory(NAME_Zlib, CompressedDataRaw, OutCompressedData.CompressedSize, UncompressedData.GetData(), OutCompressedData.UncompressedSize);
+	const bool bResult = FCompression::CompressMemory(NAME_Zlib, CompressedDataRaw.GetData(), OutCompressedData.CompressedSize, UncompressedData.GetData(), OutCompressedData.UncompressedSize);
 	if (bResult)
 	{
 		// Copy compressed data into the array.
-		OutCompressedData.Data.Append(CompressedDataRaw, OutCompressedData.CompressedSize);
-		delete[] CompressedDataRaw;
-		CompressedDataRaw = nullptr;
+		OutCompressedData.Data.Append(CompressedDataRaw.GetData(), OutCompressedData.CompressedSize);
 	}
 
 	return bResult;
@@ -540,7 +541,7 @@ void FCrashUploadToReceiver::PostReportComplete()
 	{
 #if PRIMARY_UPLOAD_RECEIVER
 		// completed upload to CRR so send analytics
-		FPrimaryCrashProperties::Get()->SendPostUploadAnalytics();
+		FPrimaryCrashProperties::Get()->SendPostUploadAnalytics(0.0, false, 0, 0);
 #endif
 		SetCurrentState(EUploadState::PostingReportComplete);
 	}
@@ -718,6 +719,7 @@ bool FCrashUploadToReceiver::ParseServerResponse(FHttpResponsePtr Response, bool
 
 FCrashUploadToDataRouter::FCrashUploadToDataRouter(const FString& InDataRouterUrl)
 	: DataRouterUrl(InDataRouterUrl)
+	, Timer(Duration)
 {
 	if (!DataRouterUrl.IsEmpty())
 	{
@@ -736,9 +738,10 @@ FCrashUploadToDataRouter::FCrashUploadToDataRouter(const FString& InDataRouterUr
 
 FCrashUploadToDataRouter::~FCrashUploadToDataRouter()
 {
+	Timer.Stop();
 #if PRIMARY_UPLOAD_DATAROUTER
 	// completed upload to DR so send analytics
-	FPrimaryCrashProperties::Get()->SendPostUploadAnalytics();
+	FPrimaryCrashProperties::Get()->SendPostUploadAnalytics(Duration, bResult, ResponseCode, PayloadSize, ReportCount);
 #endif		
 }
 
@@ -746,6 +749,8 @@ void FCrashUploadToDataRouter::BeginUpload(const FPlatformErrorReport& PlatformE
 {
 	bUploadCalled = true;
 
+	Timer.Start();
+
 	ErrorReport = PlatformErrorReport;
 	PendingFiles = FPlatformErrorReport(ErrorReport.GetReportDirectory()).GetFilesToUpload();
 	UE_LOG(CrashReportCoreLog, Log, TEXT("Got %d pending files to upload from '%s'"), PendingFiles.Num(), *ErrorReport.GetReportDirectoryLeafName());
@@ -772,6 +777,9 @@ void FCrashUploadToDataRouter::CompressAndSendData()
 		return;
 	}
 
+	PayloadSize += CompressedData.Data.Num();
+	++ReportCount;
+
 	PendingFiles.Empty();
 
 	FString UserId = FString::Printf(TEXT("%s|%s|%s"), *FPlatformMisc::GetLoginId(), *FPlatformMisc::GetEpicAccountId(), *FPlatformMisc::GetOperatingSystemId());
@@ -791,11 +799,7 @@ void FCrashUploadToDataRouter::CompressAndSendData()
 	Request->SetContent(CompressedData.Data);
 	UE_LOG(CrashReportCoreLog, Log, TEXT("Sending HTTP request: %s, Payload size: %d"), *Request->GetURL(), CompressedData.Data.Num());
 
-	if (Request->ProcessRequest())
-	{
-		return;
-	}
-	else
+	if (!Request->ProcessRequest())
 	{
 		UE_LOG(CrashReportCoreLog, Warning, TEXT("Failed to send file upload request"));
 		SetCurrentState(EUploadState::Cancelled);
@@ -811,7 +815,18 @@ TSharedRef<IHttpRequest, ESPMode::ThreadSafe> FCrashUploadToDataRouter::CreateHt
 
 void FCrashUploadToDataRouter::OnProcessRequestComplete(FHttpRequestPtr HttpRequest, FHttpResponsePtr HttpResponse, bool bSucceeded)
 {
-	UE_LOG(CrashReportCoreLog, Log, TEXT("OnProcessRequestComplete(), State=%s bSucceeded=%i"), ToString(State), (int32)bSucceeded);
+	bResult = bSucceeded;
+	ResponseCode = HttpResponse->GetResponseCode();
+
+	UE_LOG(CrashReportCoreLog, Log, TEXT("OnProcessRequestComplete(), State=%s Response=%u ConnectedSuccesfully=%i"), ToString(State), ResponseCode, (int32)bSucceeded);
+	
+	if (!EHttpResponseCodes::IsOk(ResponseCode))
+	{
+		const FText Description = GetDescription((EHttpResponseCodes::Type)ResponseCode);
+		UE_LOG(CrashReportCoreLog, Error, TEXT("Failed to send crash report. Server returned error code %u (%s)."), ResponseCode, *Description.ToString());
+		SetCurrentState(EUploadState::Cancelled);
+	}
+	
 	switch (State)
 	{
 	default:
@@ -837,17 +852,20 @@ void FCrashUploadToDataRouter::OnProcessRequestComplete(FHttpRequestPtr HttpRequ
 
 void FCrashUploadToDataRouter::CheckPendingReportsForFilesToUpload()
 {
-	SetCurrentState(EUploadState::CompressAndSendData);
-
-	for (; PendingReportDirectoryIndex < PendingReportDirectories.Num(); PendingReportDirectoryIndex++)
+	if (!PendingReportDirectories.IsEmpty())
 	{
-		ErrorReport = FPlatformErrorReport(PendingReportDirectories[PendingReportDirectoryIndex]);
-		PendingFiles = ErrorReport.GetFilesToUpload();
-
-		if (PendingFiles.Num() > 0)
+		UE_LOG(CrashReportCoreLog, Log, TEXT("Found additional pending reports"));
+		SetCurrentState(EUploadState::CompressAndSendData);
+		for (; PendingReportDirectoryIndex < PendingReportDirectories.Num(); PendingReportDirectoryIndex++)
 		{
-			CompressAndSendData();
-			return;
+			ErrorReport = FPlatformErrorReport(PendingReportDirectories[PendingReportDirectoryIndex]);
+			PendingFiles = ErrorReport.GetFilesToUpload();
+
+			if (PendingFiles.Num() > 0)
+			{
+				CompressAndSendData();
+				return;
+			}
 		}
 	}
 
diff --git a/Engine/Source/Runtime/CrashReportCore/Public/CrashDescription.h b/Engine/Source/Runtime/CrashReportCore/Public/CrashDescription.h
--- a/Engine/Source/Runtime/CrashReportCore/Public/CrashDescription.h
+++ b/Engine/Source/Runtime/CrashReportCore/Public/CrashDescription.h
@@ -499,8 +499,8 @@ public:
 	/** Sends this crash for analytics (before upload). */
 	void SendPreUploadAnalytics();
 
-	/** Sends this crash for analytics (after successful upload). */
-	void SendPostUploadAnalytics();
+	/** Sends this crash for analytics (after upload) with additional stats about the upload. */
+	void SendPostUploadAnalytics(double Duration, bool bResult, int32 ResponseCode, uint32 PayloadSize, uint32 ReportCount);
 
 	/** Saves the data. */
 	void Save();
diff --git a/Engine/Source/Runtime/CrashReportCore/Public/CrashUpload.h b/Engine/Source/Runtime/CrashReportCore/Public/CrashUpload.h
--- a/Engine/Source/Runtime/CrashReportCore/Public/CrashUpload.h
+++ b/Engine/Source/Runtime/CrashReportCore/Public/CrashUpload.h
@@ -9,6 +9,7 @@
 #include "Interfaces/IHttpRequest.h"
 #include "Internationalization/Text.h"
 #include "PlatformErrorReport.h"
+#include "ProfilingDebugging/ScopedTimers.h"
 #include "Templates/SharedPointer.h"
 
 struct FCompressedData;
@@ -286,4 +287,22 @@ private:
 private:
 	/** Url for data router requests */
 	FString DataRouterUrl;
+
+	/** HTTP request result */
+	bool bResult = false;
+	
+	/** HTTP response code */
+	int32 ResponseCode = 0;
+
+	/** Payload size */
+	uint32 PayloadSize = 0;
+
+	/** Reports uploaded */
+	uint32 ReportCount;
+
+	/** Duration of compression + upload */
+	double Duration = 0.0;
+
+	/** Timer */
+	FDurationTimer Timer;
 };
