Remove "Classic" string in symbols version. Add PCallStackHash

diff --git a/Engine/Source/Runtime/Core/Private/GenericPlatform/GenericPlatformCrashContext.cpp b/Engine/Source/Runtime/Core/Private/GenericPlatform/GenericPlatformCrashContext.cpp
--- a/Engine/Source/Runtime/Core/Private/GenericPlatform/GenericPlatformCrashContext.cpp
+++ b/Engine/Source/Runtime/Core/Private/GenericPlatform/GenericPlatformCrashContext.cpp
@@ -294,9 +294,6 @@ void FGenericCrashContext::Initialize()
 	FString Symbols = FString(UE_SYMBOLS_VERSION);
 #else
 	FString Symbols = FString::Printf(TEXT("%s"), FApp::GetBuildVersion());
-#endif
-#ifdef UE_APP_FLAVOR
-	Symbols = FString::Printf(TEXT("%s-%s"), *Symbols, *FString(UE_APP_FLAVOR));
 #endif
 	Symbols = FString::Printf(TEXT("%s-%s-%s"), *Symbols, FPlatformMisc::GetUBTPlatform(), NCached::Session.BuildConfigurationName).Replace(TEXT("+"), TEXT("*"));
 #ifdef UE_BUILD_FLAVOR
diff --git a/Engine/Source/Runtime/CrashReportCore/Private/CrashDescription.cpp b/Engine/Source/Runtime/CrashReportCore/Private/CrashDescription.cpp
--- a/Engine/Source/Runtime/CrashReportCore/Private/CrashDescription.cpp
+++ b/Engine/Source/Runtime/CrashReportCore/Private/CrashDescription.cpp
@@ -127,6 +127,7 @@ FPrimaryCrashProperties::FPrimaryCrashProperties()
 	// Multiline properties
 	, CallStack( FGenericCrashContext::RuntimePropertiesTag, TEXT( "CallStack" ), this )
 	, PCallStack( FGenericCrashContext::RuntimePropertiesTag, TEXT( "PCallStack" ), this )
+	, PCallStackHashProperty( FGenericCrashContext::RuntimePropertiesTag, TEXT( "PCallStackHash" ), this )
 	, SourceContext( FGenericCrashContext::RuntimePropertiesTag, TEXT( "SourceContext" ), this )
 	, Modules( FGenericCrashContext::RuntimePropertiesTag, TEXT( "Modules" ), this )
 	, UserDescription( FGenericCrashContext::RuntimePropertiesTag, TEXT( "UserDescription" ), this )
diff --git a/Engine/Source/Runtime/CrashReportCore/Private/GenericErrorReport.cpp b/Engine/Source/Runtime/CrashReportCore/Private/GenericErrorReport.cpp
--- a/Engine/Source/Runtime/CrashReportCore/Private/GenericErrorReport.cpp
+++ b/Engine/Source/Runtime/CrashReportCore/Private/GenericErrorReport.cpp
@@ -195,7 +195,15 @@ void FGenericErrorReport::SetPrimaryCrashProperties( FPrimaryCrashProperties& ou
 		if (CallStack.Num() > 0)
 		{
 			out_PrimaryCrashProperties.CallStack = CallStack;
+		}
+		if (PCallStack.Num() > 0)
+		{
 			out_PrimaryCrashProperties.PCallStack = PCallStack;
+			
+			FString StringToHash = out_PrimaryCrashProperties.PCallStack.AsString();
+			FSHAHash PCallStackHash;
+			FSHA1::HashBuffer(*StringToHash, StringToHash.Len() * sizeof(FString::ElementType), PCallStackHash.Hash);
+			out_PrimaryCrashProperties.PCallStackHashProperty = PCallStackHash.ToString();
 		}
 		out_PrimaryCrashProperties.Modules = Helper->CrashInfo.ModuleNames;
 		out_PrimaryCrashProperties.SourceContext = Helper->CrashInfo.SourceContext;
diff --git a/Engine/Source/Runtime/CrashReportCore/Public/CrashDescription.h b/Engine/Source/Runtime/CrashReportCore/Public/CrashDescription.h
--- a/Engine/Source/Runtime/CrashReportCore/Public/CrashDescription.h
+++ b/Engine/Source/Runtime/CrashReportCore/Public/CrashDescription.h
@@ -319,10 +319,15 @@ struct FPrimaryCrashProperties
 	
 	/**
 	 * An array of FStrings representing the portable callstack of the crash.
-	 * @RawCallStack	varchar(MAX)
 	 *
 	 */
 	FCrashProperty PCallStack;
+	
+	/**
+	 * The hash of PCallStack
+	 *
+	 */
+	FCrashProperty PCallStackHashProperty;
 
 	/**
 	 * An array of FStrings showing the source code around the crash.
