Fix crashes where http response is null

In certain situations the response ptr can be a nullptr.

diff --git a/Engine/Source/Runtime/CrashReportCore/Private/CrashUpload.cpp b/Engine/Source/Runtime/CrashReportCore/Private/CrashUpload.cpp
--- a/Engine/Source/Runtime/CrashReportCore/Private/CrashUpload.cpp
+++ b/Engine/Source/Runtime/CrashReportCore/Private/CrashUpload.cpp
@@ -694,6 +694,10 @@ void FCrashUploadToReceiver::SendPingRequest()
 
 bool FCrashUploadToReceiver::ParseServerResponse(FHttpResponsePtr Response, bool& OutValidReport)
 {
+	if (!Response.IsValid())
+	{
+		return false;
+	}
 	// Turn the snippet into a complete XML document, to keep the XML parser happy
 	FXmlFile ParsedResponse(FString(TEXT("<Root>")) + Response->GetContentAsString() + TEXT("</Root>"), EConstructMethod::ConstructFromBuffer);
 	UE_LOG(CrashReportCoreLog, Log, TEXT("Response->GetContentAsString(): '%s'"), *Response->GetContentAsString());
@@ -816,9 +820,9 @@ TSharedRef<IHttpRequest, ESPMode::ThreadSafe> FCrashUploadToDataRouter::CreateHt
 void FCrashUploadToDataRouter::OnProcessRequestComplete(FHttpRequestPtr HttpRequest, FHttpResponsePtr HttpResponse, bool bSucceeded)
 {
 	bResult = bSucceeded;
-	ResponseCode = HttpResponse->GetResponseCode();
+	ResponseCode = HttpResponse.IsValid() ? HttpResponse->GetResponseCode() : -1;
 
-	UE_LOG(CrashReportCoreLog, Log, TEXT("OnProcessRequestComplete(), State=%s Response=%u ConnectedSuccesfully=%i"), ToString(State), ResponseCode, (int32)bSucceeded);
+	UE_LOG(CrashReportCoreLog, Log, TEXT("OnProcessRequestComplete(), State=%s Response=%i ConnectedSuccesfully=%i"), ToString(State), ResponseCode, (int32)bSucceeded);
 	
 	if (!EHttpResponseCodes::IsOk(ResponseCode))
 	{
