Fix crash report Android compile error

#rnx

diff --git a/Engine/Source/Runtime/CrashReportCore/Private/CrashDescription.cpp b/Engine/Source/Runtime/CrashReportCore/Private/CrashDescription.cpp
--- a/Engine/Source/Runtime/CrashReportCore/Private/CrashDescription.cpp
+++ b/Engine/Source/Runtime/CrashReportCore/Private/CrashDescription.cpp
@@ -109,16 +109,8 @@ FCrashPropertyXmlNode& FCrashPropertyXmlNode::operator=(const FXmlNode* Node)
 	return *this;
 }
 
-FCrashPropertyThreadsNode::FCrashPropertyThreadsNode(const FString& InMainCategory, const FString& InSecondCategory, FPrimaryCrashProperties* InOwner)
-: Owner(InOwner)
-	, MainCategory(InMainCategory)
-	, SecondCategory(InSecondCategory)
-	, bSet(false)
-{ }
-
-
 #if WITH_CRASHREPORTER
-FCrashPropertyThreadsNode& FCrashPropertyThreadsNode::operator=(const TArray<FCrashThreadInfo>& ThreadsInfo)
+FCrashPropertyXmlNode& FCrashPropertyXmlNode::operator=(const TArray<FCrashThreadInfo>& ThreadsInfo)
 {
 	if (ThreadsInfo.Num() <= 0)
 	{
@@ -147,6 +139,7 @@ FCrashPropertyThreadsNode& FCrashPropertyThreadsNode::operator=(const TArray<FCr
 }
 #endif
 
+
 /*-----------------------------------------------------------------------------
 	FPrimaryCrashProperties
 -----------------------------------------------------------------------------*/
diff --git a/Engine/Source/Runtime/CrashReportCore/Public/CrashDescription.h b/Engine/Source/Runtime/CrashReportCore/Public/CrashDescription.h
--- a/Engine/Source/Runtime/CrashReportCore/Public/CrashDescription.h
+++ b/Engine/Source/Runtime/CrashReportCore/Public/CrashDescription.h
@@ -157,35 +157,15 @@ struct FCrashPropertyXmlNode
 	friend struct FPrimaryCrashProperties;
 	/** Assignment operator for . */
 	FCrashPropertyXmlNode& operator=(const FXmlNode* Node);
-protected:
-	/** Initialization constructor. */
-	FCrashPropertyXmlNode(const FString& InMainCategory, const FString& InSecondCategory, FPrimaryCrashProperties* InOwner);
-
-protected:
-	/** Owner of the property. */
-	FPrimaryCrashProperties* Owner;
-
-	/** Main category in the crash context. */
-	FString MainCategory;
-
-	/** Second category in the crash context. */
-	FString SecondCategory;
-
-	mutable bool bSet;
-};
-
-struct FCrashPropertyThreadsNode
-{
-	friend struct FPrimaryCrashProperties;
 	
 #if WITH_CRASHREPORTER
-	/** Assignment operator for . */
-	FCrashPropertyThreadsNode& operator=(const TArray<FCrashThreadInfo>& ThreadsInfo);
+	/** Assignment operator for converting ThreadsInfo into xml nodes*/
+	FCrashPropertyXmlNode& operator=(const TArray<FCrashThreadInfo>& ThreadsInfo);
 #endif
 	
 protected:
 	/** Initialization constructor. */
-	FCrashPropertyThreadsNode(const FString& InMainCategory, const FString& InSecondCategory, FPrimaryCrashProperties* InOwner);
+	FCrashPropertyXmlNode(const FString& InMainCategory, const FString& InSecondCategory, FPrimaryCrashProperties* InOwner);
 
 protected:
 	/** Owner of the property. */
@@ -205,7 +185,6 @@ struct FPrimaryCrashProperties
 {
 	friend struct FCrashProperty;
 	friend struct FCrashPropertyXmlNode;
-	friend struct FCrashPropertyThreadsNode;
 
 	/** Version. */
 	ECrashDescVersions CrashVersion;
@@ -467,7 +446,7 @@ struct FPrimaryCrashProperties
 	/**
 	 *	Thread contexts, XML elements containing info specific to an active thread, e.g. callstacks
 	 */
-	FCrashPropertyThreadsNode Threads;
+	FCrashPropertyXmlNode Threads;
 
 	/**
 	 *	Optional additional data for platform properties
