Initialize ReportCount variable in Crash report analytics

diff --git a/Engine/Source/Runtime/CrashReportCore/Public/CrashUpload.h b/Engine/Source/Runtime/CrashReportCore/Public/CrashUpload.h
--- a/Engine/Source/Runtime/CrashReportCore/Public/CrashUpload.h
+++ b/Engine/Source/Runtime/CrashReportCore/Public/CrashUpload.h
@@ -298,7 +298,7 @@ private:
 	uint32 PayloadSize = 0;
 
 	/** Reports uploaded */
-	uint32 ReportCount;
+	uint32 ReportCount = 0;
 
 	/** Duration of compression + upload */
 	double Duration = 0.0;
