Unix socket support for CrashReportClient datarouter uploads

#rnx

diff --git a/Engine/Source/Programs/CrashReportClient/Private/CrashReportClient.cpp b/Engine/Source/Programs/CrashReportClient/Private/CrashReportClient.cpp
--- a/Engine/Source/Programs/CrashReportClient/Private/CrashReportClient.cpp
+++ b/Engine/Source/Programs/CrashReportClient/Private/CrashReportClient.cpp
@@ -41,7 +41,7 @@ FCrashReportClient::FCrashReportClient(const FPlatformErrorReport& InErrorReport
 	, DiagnoseReportTask(nullptr)
 	, ErrorReport( InErrorReport )
 	, ReceiverUploader(FCrashReportCoreConfig::Get().GetReceiverAddress())
-	, DataRouterUploader(FCrashReportCoreConfig::Get().GetDataRouterURL())
+	, DataRouterUploader(FCrashReportCoreConfig::Get().GetDataRouterURL(), FCrashReportCoreConfig::Get().GetDataRouterUnixSocket())
 	, bShouldWindowBeHidden(false)
 	, bSendData(bImplicitSend)
 	, bIsSuccesfullRestart(false)
diff --git a/Engine/Source/Runtime/CrashReportCore/Private/CrashReportCoreConfig.cpp b/Engine/Source/Runtime/CrashReportCore/Private/CrashReportCoreConfig.cpp
--- a/Engine/Source/Runtime/CrashReportCore/Private/CrashReportCoreConfig.cpp
+++ b/Engine/Source/Runtime/CrashReportCore/Private/CrashReportCoreConfig.cpp
@@ -144,6 +144,17 @@ void FCrashReportCoreConfig::SetProjectConfigOverrides(const FConfigFile& InConf
 #if PLATFORM_WINDOWS
 	InConfigFile.GetBool(*SectionName, TEXT("bIsAllowedToCopyFilesToClipboard"), bIsAllowedToCopyFilesToClipboard);
 #endif
+
+	bool bUnixSocketEnabled = false;
+	InConfigFile.GetBool(*SectionName, TEXT("bUnixSocketEnabled"), bUnixSocketEnabled); 
+	if (bUnixSocketEnabled)
+	{
+		FString UnixSocketPath;
+		InConfigFile.GetString(*SectionName, TEXT("DataRouterUnixSocketsPath"), UnixSocketPath);
+		FString UnixSocketName;
+		InConfigFile.GetString(*SectionName, TEXT("DataRouterUnixSocket"), UnixSocketName);
+		DataRouterUnixSocket = UnixSocketPath / UnixSocketName;
+	}
 }
 
 FString FCrashReportCoreConfig::GetFullCrashDumpLocationForBranch(const FString& BranchName) const
diff --git a/Engine/Source/Runtime/CrashReportCore/Private/CrashReportCoreUnattended.cpp b/Engine/Source/Runtime/CrashReportCore/Private/CrashReportCoreUnattended.cpp
--- a/Engine/Source/Runtime/CrashReportCore/Private/CrashReportCoreUnattended.cpp
+++ b/Engine/Source/Runtime/CrashReportCore/Private/CrashReportCoreUnattended.cpp
@@ -8,7 +8,7 @@
 
 FCrashReportCoreUnattended::FCrashReportCoreUnattended(FPlatformErrorReport& InErrorReport, bool InExitWhenComplete)
 	: ReceiverUploader(FCrashReportCoreConfig::Get().GetReceiverAddress())
-	, DataRouterUploader(FCrashReportCoreConfig::Get().GetDataRouterURL())
+	, DataRouterUploader(FCrashReportCoreConfig::Get().GetDataRouterURL(), FCrashReportCoreConfig::Get().GetDataRouterUnixSocket())
 	, ErrorReport(InErrorReport)
     , bUploadComplete(false)
 	, bExitWhenComplete(InExitWhenComplete)
diff --git a/Engine/Source/Runtime/CrashReportCore/Private/CrashUpload.cpp b/Engine/Source/Runtime/CrashReportCore/Private/CrashUpload.cpp
--- a/Engine/Source/Runtime/CrashReportCore/Private/CrashUpload.cpp
+++ b/Engine/Source/Runtime/CrashReportCore/Private/CrashUpload.cpp
@@ -721,8 +721,9 @@ bool FCrashUploadToReceiver::ParseServerResponse(FHttpResponsePtr Response, bool
 
 // FCrashUploadToDataRouter //////////////////////////////////////////////////////
 
-FCrashUploadToDataRouter::FCrashUploadToDataRouter(const FString& InDataRouterUrl)
+FCrashUploadToDataRouter::FCrashUploadToDataRouter(const FString& InDataRouterUrl, const FString& InOptionalDataRouterUnixSocket)
 	: DataRouterUrl(InDataRouterUrl)
+	, DataRouterUnixSocket(InOptionalDataRouterUnixSocket)
 	, Timer(Duration)
 {
 	if (!DataRouterUrl.IsEmpty())
@@ -800,6 +801,12 @@ void FCrashUploadToDataRouter::CompressAndSendData()
 	Request->SetVerb(TEXT("POST"));
 	Request->SetHeader(TEXT("Content-Type"), TEXT("application/octet-stream"));
 	Request->SetURL(DataRouterUrl + UrlParams);
+#if UE_HTTP_SUPPORT_UNIX_SOCKET
+	if (!DataRouterUnixSocket.IsEmpty())
+	{
+		Request->SetOption(HttpRequestOptions::UnixSocketPath, DataRouterUnixSocket);
+	}
+#endif
 	Request->SetContent(CompressedData.Data);
 	UE_LOG(CrashReportCoreLog, Log, TEXT("Sending HTTP request: %s, Payload size: %d"), *Request->GetURL(), CompressedData.Data.Num());
 
diff --git a/Engine/Source/Runtime/CrashReportCore/Public/CrashReportCoreConfig.h b/Engine/Source/Runtime/CrashReportCore/Public/CrashReportCoreConfig.h
--- a/Engine/Source/Runtime/CrashReportCore/Public/CrashReportCoreConfig.h
+++ b/Engine/Source/Runtime/CrashReportCore/Public/CrashReportCoreConfig.h
@@ -75,6 +75,11 @@ struct FCrashReportCoreConfig
 		return DataRouterUrl;
 	}
 
+	const FString& GetDataRouterUnixSocket() const
+	{
+		return DataRouterUnixSocket;
+	}
+
 	const FString& GetDiagnosticsFilename() const
 	{
 		return DiagnosticsFilename;
@@ -159,6 +164,9 @@ protected:
 	/** URL of Data Router service */
 	FString DataRouterUrl;
 
+	/** Local Unix Socket to use for Data Router service */
+	FString DataRouterUnixSocket;
+
 	/** Filename to use when saving diagnostics report, if generated locally. */
 	FString DiagnosticsFilename;
 
diff --git a/Engine/Source/Runtime/CrashReportCore/Public/CrashUpload.h b/Engine/Source/Runtime/CrashReportCore/Public/CrashUpload.h
--- a/Engine/Source/Runtime/CrashReportCore/Public/CrashUpload.h
+++ b/Engine/Source/Runtime/CrashReportCore/Public/CrashUpload.h
@@ -251,7 +251,7 @@ public:
 	 * Constructor: pings server
 	 * @param ServerAddress Host IP of the crash report server
 	 */
-	explicit FCrashUploadToDataRouter(const FString& InDataRouterUrl);
+	explicit FCrashUploadToDataRouter(const FString& InDataRouterUrl, const FString& InOptionalDataRouterUnixSocket);
 
 	/**
 	 * Destructor for logging
@@ -288,6 +288,9 @@ private:
 	/** Url for data router requests */
 	FString DataRouterUrl;
 
+	/** Local Unix Socket for data router */
+	FString DataRouterUnixSocket;
+
 	/** HTTP request result */
 	bool bResult = false;
 	
