[CrashReportClient] Additional CrashReportClient customization
Added 2 new options:
* option for hiding the allow to be contacted text
* option for customizing whether the user can send a crash without providing detailed information

diff --git a/Engine/Config/BaseEngine.ini b/Engine/Config/BaseEngine.ini
--- a/Engine/Config/BaseEngine.ini
+++ b/Engine/Config/BaseEngine.ini
@@ -3593,7 +3593,9 @@ ConnectionRetryDelay=2
 
 [CrashReportClient]
 bHideLogFilesOption=false
+bHideAllowToBeContactedOption=false
 bIsAllowedToCloseWithoutSending=true
+bIsAllowedToSendWithoutDetailedInfo=true
 CrashConfigPurgeDays=2
 Stall.RecordDump=false
 Ensure.RecordDump=true
diff --git a/Engine/Source/Programs/CrashReportClient/Private/SCrashReportClient.cpp b/Engine/Source/Programs/CrashReportClient/Private/SCrashReportClient.cpp
--- a/Engine/Source/Programs/CrashReportClient/Private/SCrashReportClient.cpp
+++ b/Engine/Source/Programs/CrashReportClient/Private/SCrashReportClient.cpp
@@ -47,6 +47,7 @@ void SCrashReportClient::Construct(const FArguments& InArgs, const TSharedRef<FC
 {
 	CrashReportClient = Client;
 	bHasUserCommentErrors = false;
+	bIsUserCommentEmpty = true;
 	bHideSubmitAndRestart = InArgs._bHideSubmitAndRestart;
 
 	FText CrashDetailedMessage = LOCTEXT("CrashDetailed", "We are very sorry that this crash occurred. Our goal is to prevent crashes like this from occurring in the future. Please help us track down and fix this crash by providing detailed information about what you were doing so that we may reproduce the crash and fix it quickly. You can also log a Bug Report with us using the <a id=\"browser\" href=\"https://epicsupport.force.com/unrealengine/s/\" style=\"Hyperlink\">Bug Submission Form</> and work directly with support staff to report this issue.\n\nThanks for your help in improving the Unreal Engine.");
@@ -304,6 +305,7 @@ void SCrashReportClient::ConstructDetailedDialog(const TSharedRef<FCrashReportCl
 			.Padding( FMargin( 4, 4 ) )
 			[
 				SNew(SHorizontalBox)
+				.Visibility( FCrashReportCoreConfig::Get().GetHideAllowToBeContactedOption() ? EVisibility::Collapsed : EVisibility::Visible )
 
 				+SHorizontalBox::Slot()
 				.AutoWidth()
@@ -554,6 +556,7 @@ void SCrashReportClient::OnUserCommentTextChanged(const FText& NewText)
 {
 	FText ErrorMessage = FText::GetEmpty();
 	bHasUserCommentErrors = false;
+	bIsUserCommentEmpty = NewText.IsEmpty();
 
 	int SizeLimit = FCrashReportCoreConfig::Get().GetUserCommentSizeLimit();
 	int Size = NewText.ToString().Len();
@@ -575,8 +578,11 @@ bool SCrashReportClient::IsSendEnabled() const
 {
 	const bool bValidAppName = FPrimaryCrashProperties::Get()->IsValid() && !FPrimaryCrashProperties::Get()->GameName.IsEmpty();
 	const bool bValidEndPoint = !FCrashReportCoreConfig::Get().GetReceiverAddress().IsEmpty() || !FCrashReportCoreConfig::Get().GetDataRouterURL().IsEmpty();
+	const bool bIsAllowedToSendWithoutDetailedInfo = FCrashReportCoreConfig::Get().IsAllowedToSendWithoutDetailedInfo();
 
-	return bValidAppName && bValidEndPoint && !bHasUserCommentErrors;
+	const bool bHasValidInput = bValidAppName && bValidEndPoint && !bHasUserCommentErrors;
+	const bool bCanSend = bIsAllowedToSendWithoutDetailedInfo || !bIsUserCommentEmpty;
+	return bHasValidInput && bCanSend;
 }
 
 FText SCrashReportClient::GetSendTooltip()
diff --git a/Engine/Source/Programs/CrashReportClient/Private/SCrashReportClient.h b/Engine/Source/Programs/CrashReportClient/Private/SCrashReportClient.h
--- a/Engine/Source/Programs/CrashReportClient/Private/SCrashReportClient.h
+++ b/Engine/Source/Programs/CrashReportClient/Private/SCrashReportClient.h
@@ -84,6 +84,7 @@ private:
 
 	TSharedPtr<SMultiLineEditableTextBox> CrashDetailsInformation;
 
+	bool bIsUserCommentEmpty;
 	bool bHasUserCommentErrors;
 	bool bHideSubmitAndRestart;
 
diff --git a/Engine/Source/Runtime/CrashReportCore/Private/CrashReportCoreConfig.cpp b/Engine/Source/Runtime/CrashReportCore/Private/CrashReportCoreConfig.cpp
--- a/Engine/Source/Runtime/CrashReportCore/Private/CrashReportCoreConfig.cpp
+++ b/Engine/Source/Runtime/CrashReportCore/Private/CrashReportCoreConfig.cpp
@@ -51,6 +51,11 @@ FCrashReportCoreConfig::FCrashReportCoreConfig()
 		}
 	}
 
+	if (!GConfig->GetBool(*SectionName, TEXT("bHideAllowToBeContactedOption"), bHideAllowToBeContactedOption, GEngineIni))
+	{
+		bHideAllowToBeContactedOption = false;
+	}
+
 	if (!GConfig->GetBool(*SectionName, TEXT("bSendLogFile"), bSendLogFile, GEngineIni ))
 	{
 		// Default to true when unattended when config is missing. This is mostly for dedicated servers that do not have config files for CRC.
@@ -74,13 +79,18 @@ FCrashReportCoreConfig::FCrashReportCoreConfig()
 	{
 		bHideRestartOption = false;
 	}
-	
+
 	if (!GConfig->GetBool(*SectionName, TEXT("bIsAllowedToCloseWithoutSending"), bIsAllowedToCloseWithoutSending, GEngineIni))
 	{
 		// Default to false (show the option) when config is missing.
 		bIsAllowedToCloseWithoutSending = true;
 	}
-	
+
+	if (!GConfig->GetBool(*SectionName, TEXT("bIsAllowedToSendWithoutDetailedInfo"), bIsAllowedToSendWithoutDetailedInfo, GEngineIni))
+	{
+		bIsAllowedToSendWithoutDetailedInfo = true;
+	}
+
 	if (!GConfig->GetBool(*SectionName, TEXT("bShowEndpointInTooltip"), bShowEndpointInTooltip, GEngineIni))
 	{
 		bShowEndpointInTooltip = false;
@@ -134,11 +144,13 @@ void FCrashReportCoreConfig::SetProjectConfigOverrides(const FConfigFile& InConf
 	InConfigFile.GetString(*SectionName, TEXT("CrashReportReceiverIP"), CrashReportReceiverIP);
 	InConfigFile.GetString(*SectionName, TEXT("DataRouterUrl"), DataRouterUrl);
 	InConfigFile.GetBool(*SectionName, TEXT("bAllowToBeContacted"), bAllowToBeContacted);
+	InConfigFile.GetBool(*SectionName, TEXT("bHideAllowToBeContactedOption"), bHideAllowToBeContactedOption);
 	InConfigFile.GetBool(*SectionName, TEXT( "bSendLogFile" ), bSendLogFile);
 	InConfigFile.GetInt(*SectionName, TEXT("UserCommentSizeLimit"), UserCommentSizeLimit);
 	InConfigFile.GetBool(*SectionName, TEXT("bHideLogFilesOption"), bHideLogFilesOption);
 	InConfigFile.GetBool(*SectionName, TEXT("bHideRestartOption"), bHideRestartOption);
 	InConfigFile.GetBool(*SectionName, TEXT("bIsAllowedToCloseWithoutSending"), bIsAllowedToCloseWithoutSending);
+	InConfigFile.GetBool(*SectionName, TEXT("bIsAllowedToSendWithoutDetailedInfo"), bIsAllowedToSendWithoutDetailedInfo);
 	InConfigFile.GetBool(*SectionName, TEXT("bShowEndpointInTooltip"), bShowEndpointInTooltip);
 	InConfigFile.GetString(*SectionName, TEXT("CompanyName"), CompanyName);
 #if PLATFORM_WINDOWS
diff --git a/Engine/Source/Runtime/CrashReportCore/Public/CrashReportCoreConfig.h b/Engine/Source/Runtime/CrashReportCore/Public/CrashReportCoreConfig.h
--- a/Engine/Source/Runtime/CrashReportCore/Public/CrashReportCoreConfig.h
+++ b/Engine/Source/Runtime/CrashReportCore/Public/CrashReportCoreConfig.h
@@ -90,6 +90,11 @@ struct FCrashReportCoreConfig
 		return bAllowToBeContacted;
 	}
 
+	bool GetHideAllowToBeContactedOption() const
+	{
+		return bHideAllowToBeContactedOption;
+	}
+
 	bool GetSendLogFile() const
 	{
 		return bSendLogFile;
@@ -110,6 +115,11 @@ struct FCrashReportCoreConfig
 		return bIsAllowedToCloseWithoutSending;
 	}
 
+	bool IsAllowedToSendWithoutDetailedInfo() const
+	{
+		return bIsAllowedToSendWithoutDetailedInfo;
+	}
+
 	bool GetShowEndpointInTooltip() const
 	{
 		return bShowEndpointInTooltip;
@@ -183,6 +193,9 @@ protected:
 	*/
 	bool bAllowToBeContacted;
 
+	/** Whether the user is shown the option to be contacted. */
+	bool bHideAllowToBeContactedOption;
+
 	/** Whether the user allowed us to send the log file. */
 	bool bSendLogFile;
 
@@ -195,6 +208,9 @@ protected:
 	/** Whether the user is allowed to close the crash reporter without sending a report */
 	bool bIsAllowedToCloseWithoutSending;
 
+	/** Whether the user is allowed to send the crash without providing detailed information about the crash */
+	bool bIsAllowedToSendWithoutDetailedInfo;
+
 	/** Optional string displayed in the "allow to be contacted" text. */
 	FString CompanyName;
 
