Ensure only data containing the 'Threads' tag can be merged onto Android crash report's Thread node.

diff --git a/Engine/Source/Runtime/CrashReportCore/Private/Android/AndroidErrorReport.cpp b/Engine/Source/Runtime/CrashReportCore/Private/Android/AndroidErrorReport.cpp
--- a/Engine/Source/Runtime/CrashReportCore/Private/Android/AndroidErrorReport.cpp
+++ b/Engine/Source/Runtime/CrashReportCore/Private/Android/AndroidErrorReport.cpp
@@ -64,22 +64,36 @@ FText FAndroidErrorReport::DiagnoseReport() const
 		FXmlFile MergeNode(MergeFilePathName);
 		if (MergeNode.IsValid())
 		{
-			switch (MergeType)
+			bool bSuccess = false;
+			FXmlNode* Root = MergeNode.GetRootNode();
+			if(ensure(Root))
 			{
-				case FMergeInfo::EMergeType::Threads:
+				switch (MergeType)
 				{
-					FPrimaryCrashProperties::Get()->Threads = MergeNode.GetRootNode();
-					break;
-				}
-				case FMergeInfo::EMergeType::PlatformProperties:
-				{
-					FPrimaryCrashProperties::Get()->PlatformPropertiesExtras = MergeNode.GetRootNode();
-					break;
+					case FMergeInfo::EMergeType::Threads:
+					{
+						if (Root->GetTag().Contains(TEXT("Threads")))
+						{
+							FPrimaryCrashProperties::Get()->Threads = Root;
+							bSuccess = true;
+						}
+						break;
+					}
+					case FMergeInfo::EMergeType::PlatformProperties:
+					{
+						FPrimaryCrashProperties::Get()->PlatformPropertiesExtras = Root;
+						bSuccess = true;
+						break;
+					}
 				}
 			}
-			// delete the file as it has been added to the primary report.
-			IPlatformFile& PlatformFile = FPlatformFileManager::Get().GetPlatformFile();
-			PlatformFile.DeleteFile(*MergeFilePathName);
+
+			if( bSuccess ) // Keep unmerged data for inspection.
+			{
+				// delete the file as it has been added to the primary report.
+				IPlatformFile& PlatformFile = FPlatformFileManager::Get().GetPlatformFile();
+				PlatformFile.DeleteFile(*MergeFilePathName);
+			}
 		}
 	};
 
