Adding platform specific crash data to our crash analytics.

diff --git a/Engine/Source/Runtime/CrashReportCore/Private/CrashDescription.cpp b/Engine/Source/Runtime/CrashReportCore/Private/CrashDescription.cpp
--- a/Engine/Source/Runtime/CrashReportCore/Private/CrashDescription.cpp
+++ b/Engine/Source/Runtime/CrashReportCore/Private/CrashDescription.cpp
@@ -464,9 +464,25 @@ void FPrimaryCrashProperties::MakeCrashEventAttributes(TArray<FAnalyticsEventAtt
 	OutCrashAttributes.Add(FAnalyticsEventAttribute(TEXT("bInBackground"), bInBackground ? TEXT("true") : TEXT("false")));
 	OutCrashAttributes.Add(FAnalyticsEventAttribute(TEXT("bIsRequestingExit"), bIsRequestingExit ? TEXT("true") : TEXT("false")));
 
-	// Add arbitrary engine data
 	if (XmlFile->IsValid())
 	{
+		// Add arbitrary platform specific data
+		const FXmlNode* PlatformPropertiesNode = XmlFile->GetRootNode()->FindChildNode( FGenericCrashContext::PlatformPropertiesTag );
+		if (PlatformPropertiesNode)
+		{
+			FString PlatformPropertiesAttributes;
+			CrashDescriptionInternal::FAnalyticsJsonWriter JsonWriter(&PlatformPropertiesAttributes);
+			JsonWriter.WriteObjectStart();
+			for (const FXmlNode* ChildNode : PlatformPropertiesNode->GetChildrenNodes())
+			{
+				JsonWriter.WriteValue(ChildNode->GetTag(), ChildNode->GetContent());
+			}
+			JsonWriter.WriteObjectEnd();
+			JsonWriter.Close();
+			OutCrashAttributes.Add(FAnalyticsEventAttribute(TEXT("PlatformProperties"), FJsonFragment(MoveTemp(PlatformPropertiesAttributes))));
+		}
+		
+		// Add arbitrary engine data
 		const FXmlNode* EngineNode = XmlFile->GetRootNode()->FindChildNode( FGenericCrashContext::EngineDataTag );
 		if (EngineNode)
 		{
