[CrashReport] Read the IsRequestingExit property from RuntimeProperties and not GameData.

diff --git a/Engine/Source/Runtime/CrashReportCore/Private/CrashDescription.cpp b/Engine/Source/Runtime/CrashReportCore/Private/CrashDescription.cpp
--- a/Engine/Source/Runtime/CrashReportCore/Private/CrashDescription.cpp
+++ b/Engine/Source/Runtime/CrashReportCore/Private/CrashDescription.cpp
@@ -590,7 +590,7 @@ void FCrashContext::SetupPrimaryCrashProperties()
 		GetCrashProperty( bIsOOM, FGenericCrashContext::RuntimePropertiesTag, TEXT("MemoryStats.bIsOOM"));
 		GetCrashProperty( bLowMemoryWarning, FGenericCrashContext::GameDataTag, TEXT("bLowMemoryCalled"));
 		GetCrashProperty( bInBackground, FGenericCrashContext::GameDataTag, TEXT("bInBackground"));
-		GetCrashProperty( bIsRequestingExit, FGenericCrashContext::GameDataTag, TEXT("IsRequestingExit"));
+		GetCrashProperty( bIsRequestingExit, FGenericCrashContext::RuntimePropertiesTag, TEXT("IsRequestingExit"));
 		GetCrashProperty( CrashScreenshotWidth, FGenericCrashContext::EngineDataTag, TEXT("CrashScreenshotWidth"));
 		GetCrashProperty( CrashScreenshotHeight, FGenericCrashContext::EngineDataTag, TEXT("CrashScreenshotHeight"));
 
